Tool menu customization GUI

diff --git a/Engine/Source/Developer/ToolMenus/Private/ToolMenu.cpp b/Engine/Source/Developer/ToolMenus/Private/ToolMenu.cpp
--- a/Engine/Source/Developer/ToolMenus/Private/ToolMenu.cpp
+++ b/Engine/Source/Developer/ToolMenus/Private/ToolMenu.cpp
@@ -47,6 +47,9 @@ void UToolMenu::InitGeneratedCopy(const UToolMenu* Source, const FName InMenuNam
 	bToolBarForceSmallIcons = Source->bToolBarForceSmallIcons;
 	MenuOwner = Source->MenuOwner;
 
+	SubMenuParent = Source->SubMenuParent;
+	SubMenuSourceEntryName = Source->SubMenuSourceEntryName;
+
 	if (InContext)
 	{
 		Context = *InContext;
@@ -138,6 +141,7 @@ FToolMenuSection& UToolMenu::AddSection(const FName SectionName, const TAttribut
 
 	FToolMenuSection& NewSection = Sections.AddDefaulted_GetRef();
 	NewSection.InitSection(SectionName, InLabel, InPosition);
+	//NewSection.OwnerMenu = this;
 	return NewSection;
 }
 
@@ -218,6 +222,32 @@ bool UToolMenu::FindEntry(const FName EntryName, int32& SectionIndex, int32& Ent
 	return false;
 }
 
+const FToolMenuEntry* UToolMenu::FindEntry(const FName EntryName) const
+{
+	for (int32 i=0; i < Sections.Num(); ++i)
+	{
+		if (const FToolMenuEntry* Found = Sections[i].FindEntry(EntryName))
+		{
+			return Found;
+		}
+	}
+
+	return nullptr;
+}
+
+FToolMenuEntry* UToolMenu::FindEntry(const FName EntryName)
+{
+	for (int32 i=0; i < Sections.Num(); ++i)
+	{
+		if (FToolMenuEntry* Found = Sections[i].FindEntry(EntryName))
+		{
+			return Found;
+		}
+	}
+
+	return nullptr;
+}
+
 void UToolMenu::AddMenuEntry(const FName SectionName, const FToolMenuEntry& Args)
 {
 	FindOrAddSection(SectionName).AddEntry(Args);
@@ -241,6 +271,38 @@ FName UToolMenu::GetSectionName(const FName InEntryName) const
 	return NAME_None;
 }
 
+bool UToolMenu::ContainsSection(const FName InName) const
+{
+	if (InName != NAME_None)
+	{
+		for (const FToolMenuSection& Section : Sections)
+		{
+			if (Section.Name == InName)
+			{
+				return true;
+			}
+		}
+	}
+
+	return false;
+}
+
+bool UToolMenu::ContainsEntry(const FName InName) const
+{
+	if (InName != NAME_None)
+	{
+		for (const FToolMenuSection& Section : Sections)
+		{
+			if (Section.FindEntry(InName) != nullptr)
+			{
+				return true;
+			}
+		}
+	}
+
+	return false;
+}
+
 FCustomizedToolMenu* UToolMenu::FindMenuCustomization() const
 {
 	return UToolMenus::Get()->FindMenuCustomization(MenuName);
@@ -250,4 +312,136 @@ FCustomizedToolMenu* UToolMenu::AddMenuCustomization() const
 {
 	return UToolMenus::Get()->AddMenuCustomization(MenuName);
 }
+	
+TArray<FName> UToolMenu::GetMenuHierarchyNames(bool bIncludeSubMenuRoot) const
+{
+	TArray<FName> HierarchyNames;
+
+	TArray<UToolMenu*> Hierarchy = UToolMenus::Get()->CollectHierarchy(GetMenuName());
+	for (int32 i = Hierarchy.Num() - 1; i >= 0; --i)
+	{
+		HierarchyNames.AddUnique(Hierarchy[i]->GetMenuName());
+	}
+
+	if (bIncludeSubMenuRoot && SubMenuParent)
+	{
+		TArray<const UToolMenu*> SubMenuChain = GetSubMenuChain();
+		if (SubMenuChain.Num() > 0)
+		{
+			FString SubMenuFullPath;
+			for (int32 i = 1; i < SubMenuChain.Num(); ++i)
+			{
+				if (SubMenuFullPath.Len() > 0)
+				{
+					SubMenuFullPath += TEXT(".");
+				}
+				SubMenuFullPath += SubMenuChain[i]->SubMenuSourceEntryName.ToString();
+			}
+
+			// Hierarchy of the initial menu opened in the sub-menu chain of menus
+			TArray<UToolMenu*> FirstMenuHierarchy = UToolMenus::Get()->CollectHierarchy(SubMenuChain[0]->GetMenuName());
+			for (int32 i = FirstMenuHierarchy.Num() - 1; i >= 0; --i)
+			{
+				HierarchyNames.AddUnique(UToolMenus::JoinMenuPaths(FirstMenuHierarchy[i]->GetMenuName(), *SubMenuFullPath));
+			}
+		}
+	}	
+	Algo::Reverse(HierarchyNames);
+
+	return HierarchyNames;
+}
+
+FCustomizedToolMenuHierarchy UToolMenu::GetMenuCustomizationHierarchy() const
+{
+	FCustomizedToolMenuHierarchy Result;
+	
+	UToolMenus* ToolMenus = UToolMenus::Get();
+	TArray<FName> HierarchyNames = GetMenuHierarchyNames(true);
+	for (const FName ItName : HierarchyNames)
+	{
+		if (FCustomizedToolMenu* Found = ToolMenus->FindMenuCustomization(ItName))
+		{
+			Result.Hierarchy.Add(Found);
+		}
+	}
+
+	return Result;
+}
+
+void UToolMenu::UpdateMenuCustomizationFromMultibox(const TSharedRef<const FMultiBox>& InMultiBox)
+{
+	FCustomizedToolMenu* Customization = AddMenuCustomization();
+
+	Customization->EntryOrder.Reset();
+	Customization->SectionOrder.Reset();
+
+	FName CurrentSectionName = NAME_None;
+	const TArray< TSharedRef< const FMultiBlock > >& Blocks = InMultiBox->GetBlocks();
+	for (int32 BlockIndex = 0; BlockIndex < Blocks.Num(); ++BlockIndex)
+	{
+		const TSharedRef< const FMultiBlock >& Block = Blocks[BlockIndex];
+
+		if (Block->GetExtensionHook() == NAME_None)
+		{
+			continue;
+		}
+
+		// Ignore separators that are part of a section heading
+		if (Block->IsSeparator() && (BlockIndex + 1 < Blocks.Num()) && Blocks[BlockIndex + 1]->GetType() == EMultiBlockType::Heading)
+		{
+			continue;
+		}
+
+		if (Block->GetType() == EMultiBlockType::Heading)
+		{
+			CurrentSectionName = Block->GetExtensionHook();
+			Customization->SectionOrder.Add(CurrentSectionName);
+		}
+		else if (CurrentSectionName != NAME_None)
+		{
+			FCustomizedToolMenuNameArray& EntryOrderForSection = Customization->EntryOrder.FindOrAdd(CurrentSectionName);
+			EntryOrderForSection.Names.Add(Block->GetExtensionHook());
+		}
+	}
+}
+
+TArray<const UToolMenu*> UToolMenu::GetSubMenuChain() const
+{
+	TArray<const UToolMenu*> SubMenuChain;
+
+	TSet<const UToolMenu*> SubMenus;
+	for (const UToolMenu* CurrentMenu = this; CurrentMenu; CurrentMenu = CurrentMenu->SubMenuParent)
+	{
+		bool bIsAlreadyInSet = false;
+		SubMenus.Add(CurrentMenu, &bIsAlreadyInSet);
+		if (bIsAlreadyInSet)
+		{
+			ensure(!bIsAlreadyInSet);
+			break;
+		}
+		SubMenuChain.Add(CurrentMenu);
+	}
+
+	Algo::Reverse(SubMenuChain);
+
+	return SubMenuChain;
+}
+
+FString UToolMenu::GetSubMenuNamePath() const
+{
+	FString SubMenuNamePath;
+	TArray<const UToolMenu*> SubMenuChain = GetSubMenuChain();
+	if (SubMenuChain.Num() > 0)
+	{
+		for (int32 i = 1; i < SubMenuChain.Num(); ++i)
+		{
+			if (SubMenuNamePath.Len() > 0)
+			{
+				SubMenuNamePath += TEXT(".");
+			}
+			SubMenuNamePath += SubMenuChain[i]->SubMenuSourceEntryName.ToString();
+		}
+	}
 
+	return SubMenuNamePath;
+}
\ No newline at end of file
diff --git a/Engine/Source/Developer/ToolMenus/Private/ToolMenuSection.cpp b/Engine/Source/Developer/ToolMenus/Private/ToolMenuSection.cpp
--- a/Engine/Source/Developer/ToolMenus/Private/ToolMenuSection.cpp
+++ b/Engine/Source/Developer/ToolMenus/Private/ToolMenuSection.cpp
@@ -131,6 +131,19 @@ FToolMenuEntry* FToolMenuSection::FindEntry(const FName InName)
 	return nullptr;
 }
 
+const FToolMenuEntry* FToolMenuSection::FindEntry(const FName InName) const
+{
+	for (int32 i=0; i < Blocks.Num(); ++i)
+	{
+		if (Blocks[i].Name == InName)
+		{
+			return &Blocks[i];
+		}
+	}
+
+	return nullptr;
+}
+
 int32 FToolMenuSection::IndexOfBlock(const FName InName) const
 {
 	for (int32 i=0; i < Blocks.Num(); ++i)
diff --git a/Engine/Source/Developer/ToolMenus/Private/ToolMenus.cpp b/Engine/Source/Developer/ToolMenus/Private/ToolMenus.cpp
--- a/Engine/Source/Developer/ToolMenus/Private/ToolMenus.cpp
+++ b/Engine/Source/Developer/ToolMenus/Private/ToolMenus.cpp
@@ -70,7 +70,8 @@ FToolUIActionChoice::FToolUIActionChoice(const TSharedPtr< const FUICommandInfo
 UToolMenus::UToolMenus() :
 	bNextTickTimerIsSet(false),
 	bRefreshWidgetsNextTick(false),
-	bCleanupStaleWidgetsNextTick(false)
+	bCleanupStaleWidgetsNextTick(false),
+	bEditMenusMode(false)
 {
 }
 
@@ -115,6 +116,23 @@ FName UToolMenus::JoinMenuPaths(const FName Base, const FName Child)
 	return *(Base.ToString() + TEXT(".") + Child.ToString());
 }
 
+bool UToolMenus::SplitMenuPath(const FName MenuPath, FName& OutLeft, FName& OutRight)
+{
+	if (MenuPath != NAME_None)
+	{
+		FString Left;
+		FString Right;
+		if (MenuPath.ToString().Split(TEXT("."), &Left, &Right, ESearchCase::CaseSensitive, ESearchDir::FromEnd))
+		{
+			OutLeft = *Left;
+			OutRight = *Right;
+			return true;
+		}
+	}
+
+	return false;
+}
+
 bool UToolMenus::GetDisplayUIExtensionPoints() const
 {
 	return ShouldDisplayExtensionPoints.IsBound() && ShouldDisplayExtensionPoints.Execute();
@@ -400,6 +418,20 @@ void UToolMenus::AssembleMenu(UToolMenu* GeneratedMenu, const UToolMenu* Other)
 	}
 }
 
+bool UToolMenus::GetEditMenusMode() const
+{
+	return bEditMenusMode;
+}
+
+void UToolMenus::SetEditMenusMode(bool bShow)
+{
+	if (bEditMenusMode != bShow)
+	{
+		bEditMenusMode = bShow;
+		RefreshAllWidgets();
+	}
+}
+
 void UToolMenus::RemoveCustomization(const FName InName)
 {
 	int32 FoundIndex = FindMenuCustomizationIndex(InName);
@@ -451,61 +483,103 @@ FCustomizedToolMenu* UToolMenus::AddMenuCustomization(const FName InName)
 
 void UToolMenus::ApplyCustomization(UToolMenu* GeneratedMenu)
 {
-	FCustomizedToolMenu* CustomizedMenu = FindMenuCustomization(GeneratedMenu->MenuName);
-	if (!CustomizedMenu)
+	FCustomizedToolMenuHierarchy CustomizationHierarchy = GeneratedMenu->GetMenuCustomizationHierarchy();
+	if (CustomizationHierarchy.Hierarchy.Num() == 0)
 	{
 		return;
 	}
+	FCustomizedToolMenu CustomizedMenu = CustomizationHierarchy.GenerateFlattened();
 
 	TArray<FToolMenuSection> NewSections;
 	NewSections.Reserve(GeneratedMenu->Sections.Num());
 
-	for (const FCustomizedToolMenuSection& CustomizedSection : CustomizedMenu->Sections)
+	TSet<FName> PlacedEntries;
+
+	TArray<int32> NewSectionIndices;
+	NewSectionIndices.Reserve(GeneratedMenu->Sections.Num());
+
+	// Add sections with customized ordering first
+	for (const FName SectionName : CustomizedMenu.SectionOrder)
 	{
-		int32 SectionIndex = GeneratedMenu->IndexOfSection(CustomizedSection.Name);
-		if (SectionIndex == INDEX_NONE)
+		if (SectionName != NAME_None)
 		{
-			continue;
+			int32 OriginalIndex = GeneratedMenu->Sections.IndexOfByPredicate([SectionName](const FToolMenuSection& OriginalSection) { return OriginalSection.Name == SectionName; });
+			if (OriginalIndex != INDEX_NONE)
+			{
+				NewSectionIndices.Add(OriginalIndex);
+			}
 		}
+	}
 
-		FToolMenuSection& Section = GeneratedMenu->Sections[SectionIndex];
+	// Remaining sections get added to the end
+	for (int32 i = 0; i < GeneratedMenu->Sections.Num(); ++i)
+	{
+		NewSectionIndices.AddUnique(i);
+	}
+
+	// Copy sections
+	for (int32 i = 0; i < NewSectionIndices.Num(); ++i)
+	{
+		FToolMenuSection& NewSection = NewSections.Add_GetRef(GeneratedMenu->Sections[NewSectionIndices[i]]);
+		NewSection.Blocks.Reset();
+	}
 
-		TArray<FToolMenuEntry> NewBlocks;
-		NewBlocks.Reserve(Section.Blocks.Num());
+	// Add entries placed by customization
+	for (int32 i = 0; i < NewSectionIndices.Num(); ++i)
+	{
+		const FToolMenuSection& OriginalSection = GeneratedMenu->Sections[NewSectionIndices[i]];
+		FToolMenuSection& NewSection = NewSections[i];
 
-		for (const FCustomizedToolMenuEntry& Entry : CustomizedSection.Entries)
+		if (OriginalSection.Name != NAME_None)
 		{
-			int32 EntrySectionIndex = INDEX_NONE;
-			int32 EntryIndex = INDEX_NONE;
-			if (GeneratedMenu->FindEntry(Entry.Name, EntrySectionIndex, EntryIndex))
+			if (FCustomizedToolMenuNameArray* EntryOrder = CustomizedMenu.EntryOrder.Find(OriginalSection.Name))
 			{
-				NewBlocks.Add(GeneratedMenu->Sections[EntrySectionIndex].Blocks[EntryIndex]);
-				GeneratedMenu->Sections[EntrySectionIndex].Blocks.RemoveAt(EntryIndex);
+				for (const FName EntryName : EntryOrder->Names)
+				{
+					if (EntryName != NAME_None)
+					{
+						if (FToolMenuEntry* SourceEntry = GeneratedMenu->FindEntry(EntryName))
+						{
+							NewSection.Blocks.Add(*SourceEntry);
+							PlacedEntries.Add(EntryName);
+						}
+					}
+				}
 			}
 		}
+	}
 
-		// Remaining blocks
-		for (FToolMenuEntry& Block : Section.Blocks)
+	// Handle entries not placed by customization
+	for (int32 i = 0; i < NewSectionIndices.Num(); ++i)
+	{
+		const FToolMenuSection& OriginalSection = GeneratedMenu->Sections[NewSectionIndices[i]];
+		FToolMenuSection& NewSection = NewSections[i];
+
+		for (const FToolMenuEntry& OriginalEntry : OriginalSection.Blocks)
 		{
-			NewBlocks.Add(Block);
+			if (OriginalEntry.Name == NAME_None)
+			{
+				NewSection.Blocks.Add(OriginalEntry);
+			}
+			else
+			{
+				bool bAlreadyPlaced = false;
+				PlacedEntries.Add(OriginalEntry.Name, &bAlreadyPlaced);
+				if (!bAlreadyPlaced)
+				{
+					NewSection.Blocks.Add(OriginalEntry);
+				}
+			}
 		}
-
-		Section.Blocks = NewBlocks;
-
-		NewSections.Add(Section);
-		GeneratedMenu->Sections.RemoveAt(SectionIndex);
 	}
 
-	// Remaining sections
-	NewSections.Append(GeneratedMenu->Sections);
-
-	// Hide
+	// Hide sections and entries
 	if (!GeneratedMenu->IsEditing())
 	{
 		for (int32 SectionIndex = 0; SectionIndex < NewSections.Num(); ++SectionIndex)
 		{
 			FToolMenuSection& Section = NewSections[SectionIndex];
-			if (CustomizedMenu->HiddenSections.Contains(Section.Name))
+			if (CustomizedMenu.IsSectionHidden(Section.Name))
 			{
 				NewSections.RemoveAt(SectionIndex);
 				--SectionIndex;
@@ -514,7 +588,7 @@ void UToolMenus::ApplyCustomization(UToolMenu* GeneratedMenu)
 
 			for (int32 i = 0; i < Section.Blocks.Num(); ++i)
 			{
-				if (CustomizedMenu->HiddenEntries.Contains(Section.Blocks[i].Name))
+				if (CustomizedMenu.IsEntryHidden(Section.Blocks[i].Name))
 				{
 					Section.Blocks.RemoveAt(i);
 					--i;
@@ -559,35 +633,76 @@ void UToolMenus::AssembleMenuHierarchy(UToolMenu* GeneratedMenu, const TArray<UT
 	ApplyCustomization(GeneratedMenu);
 }
 
-void UToolMenus::FillMenuDynamic(FMenuBuilder& MenuBuilder, const FName SubMenuFullName, FNewToolMenuDelegate InConstructMenu, const FToolMenuContext Context)
+UToolMenu* UToolMenus::GenerateSubMenu(const UToolMenu* InGeneratedParent, const FName InBlockName)
 {
+	if (InGeneratedParent == nullptr || InBlockName == NAME_None)
+	{
+		return nullptr;
+	}
+
+	const FToolMenuEntry* Block = InGeneratedParent->FindEntry(InBlockName);
+	if (!Block)
+	{
+		return nullptr;
+	}
+
+	FName SubMenuFullName = JoinMenuPaths(InGeneratedParent->GetMenuName(), InBlockName);
+
+	// Submenus that are constructed by delegates can also be overridden by menus in the database
 	TArray<UToolMenu*> Hierarchy;
-	
-	if (SubMenuFullName != NAME_None)
 	{
-		Hierarchy = CollectHierarchy(SubMenuFullName);
+		// Walk up all parent menus trying to find a menu
+		FName BaseName = InGeneratedParent->GetMenuName();
+		while (BaseName != NAME_None)
+		{
+			FName JoinedName = JoinMenuPaths(BaseName, InBlockName);
+			if (UToolMenu* Found = FindMenu(JoinedName))
+			{
+				Hierarchy = CollectHierarchy(JoinedName);
+				break;
+			}
+
+			UToolMenu* BaseData = FindMenu(BaseName);
+			BaseName = BaseData ? BaseData->MenuParent : NAME_None;
+		}
 	}
 
 	// Construct menu using delegate and insert as root so it can be overridden
-	if (InConstructMenu.IsBound())
+	if (Block->SubMenuData.ConstructMenu.NewToolMenuDelegate.IsBound())
 	{
 		UToolMenu* Menu = NewObject<UToolMenu>(this);
-		Menu->Context = Context;
+		Menu->Context = InGeneratedParent->Context;
 		Menu->MenuName = SubMenuFullName;
-		InConstructMenu.Execute(Menu);
+
+		// Submenu specific data
+		Menu->SubMenuParent = InGeneratedParent;
+		Menu->SubMenuSourceEntryName = InBlockName;
+
+		Block->SubMenuData.ConstructMenu.NewToolMenuDelegate.Execute(Menu);
 		Menu->MenuName = SubMenuFullName;
 		Hierarchy.Insert(Menu, 0);
 	}
 
 	// Populate menu builder with final menu
-	UToolMenu* GeneratedMenu = GenerateMenu(Hierarchy, Context);
-	PopulateMenuBuilder(MenuBuilder, GeneratedMenu);
+	if (Hierarchy.Num() > 0)
+	{
+		UToolMenu* GeneratedMenu = NewObject<UToolMenu>(this);
+		GeneratedMenu->InitGeneratedCopy(Hierarchy[0], SubMenuFullName, &InGeneratedParent->Context);
+		GeneratedMenu->SubMenuParent = InGeneratedParent;
+		GeneratedMenu->SubMenuSourceEntryName = InBlockName;
+		AssembleMenuHierarchy(GeneratedMenu, Hierarchy);
+		return GeneratedMenu;
+	}
+
+	return nullptr;
 }
 
-void UToolMenus::FillMenu(class FMenuBuilder& MenuBuilder, FName InMenuName, FToolMenuContext InMenuContext)
+void UToolMenus::PopulateSubMenu(FMenuBuilder& MenuBuilder, TWeakObjectPtr<UToolMenu> InParent, const FName InBlockName)
 {
-	// Populate menu builder with final menu
-	PopulateMenuBuilder(MenuBuilder, GenerateMenu(InMenuName, InMenuContext));
+	if (UToolMenu* GeneratedMenu = GenerateSubMenu(InParent.Get(), InBlockName))
+	{
+		PopulateMenuBuilder(MenuBuilder, GeneratedMenu);
+	}
 }
 
 TSharedRef<SWidget> UToolMenus::GenerateToolbarComboButtonMenu(const FName SubMenuFullName, FToolMenuContext InContext)
@@ -613,14 +728,22 @@ void UToolMenus::FillMenuBarDropDown(class FMenuBuilder& MenuBuilder, FName InPa
 void UToolMenus::PopulateMenuBuilder(FMenuBuilder& MenuBuilder, UToolMenu* MenuData)
 {
 	const bool bIsEditing = MenuData->IsEditing();
-	if (GetDisplayUIExtensionPoints() && !bIsEditing)
+	if (GetEditMenusMode() && !bIsEditing && EditMenuDelegate.IsBound())
 	{
+		TWeakObjectPtr<UToolMenu> WeakMenuPtr = MenuData;
 		const FName MenuName = MenuData->GetMenuName();
 		MenuBuilder.AddMenuEntry(
-			FText::FromName(MenuName),
-			LOCTEXT("CopyMenuNameToClipboard", "Copy menu name to clipboard"),
-			FSlateIcon(),
-			FExecuteAction::CreateLambda([MenuName]() { FPlatformApplicationMisc::ClipboardCopy(*MenuName.ToString()); }),
+			FText::Format(LOCTEXT("EditMenu_Label", "Edit Menu: {0}"), FText::FromName(MenuName)),
+			LOCTEXT("EditMenu_ToolTip", "Open menu editor"),
+			EditMenuIcon,
+			FExecuteAction::CreateLambda([MenuName, WeakMenuPtr]()
+			{
+				FPlatformApplicationMisc::ClipboardCopy(*MenuName.ToString());
+				if (UToolMenu* InMenu = WeakMenuPtr.Get())
+				{
+					UToolMenus::Get()->EditMenuDelegate.ExecuteIfBound(InMenu);
+				}
+			}),
 			"MenuName"
 		);
 	}
@@ -665,6 +788,12 @@ void UToolMenus::PopulateMenuBuilder(FMenuBuilder& MenuBuilder, UToolMenu* MenuD
 				Widget = Block.MakeWidget.Execute(MenuData->Context);
 			}
 
+			TAttribute<FText> LabelToDisplay = Block.Label;
+			if (bIsEditing && (!Block.Label.IsSet() || Block.Label.Get().IsEmpty()))
+			{
+				LabelToDisplay = FText::FromName(Block.Name);
+			}
+
 			if (Block.Type == EMultiBlockType::MenuEntry)
 			{
 				if (Block.IsSubMenu())
@@ -675,13 +804,9 @@ void UToolMenus::PopulateMenuBuilder(FMenuBuilder& MenuBuilder, UToolMenu* MenuD
 					{
 						NewMenuDelegate = Block.SubMenuData.ConstructMenu.NewMenuDelegate;
 					}
-					else if (Block.SubMenuData.ConstructMenu.NewToolMenuDelegate.IsBound())
-					{
-						NewMenuDelegate = FNewMenuDelegate::CreateUObject(this, &UToolMenus::FillMenuDynamic, SubMenuFullName, Block.SubMenuData.ConstructMenu.NewToolMenuDelegate, MenuData->Context);
-					}
 					else
 					{
-						NewMenuDelegate = FNewMenuDelegate::CreateUObject(this, &UToolMenus::FillMenu, SubMenuFullName, MenuData->Context);
+						NewMenuDelegate = FNewMenuDelegate::CreateUObject(this, &UToolMenus::PopulateSubMenu, TWeakObjectPtr<UToolMenu>(MenuData), Block.Name);
 					}
 
 					if (Widget.IsValid())
@@ -723,7 +848,7 @@ void UToolMenus::PopulateMenuBuilder(FMenuBuilder& MenuBuilder, UToolMenu* MenuD
 							UE_LOG(LogToolMenus, Error, TEXT("UI command not found for menu entry: %s, menu: %s"), *Block.Name.ToString(), *MenuData->MenuName.ToString());
 						}
 
-						MenuBuilder.AddMenuEntry(Block.Command, Block.Name, Block.Label, Block.ToolTip, Block.Icon.Get());
+						MenuBuilder.AddMenuEntry(Block.Command, Block.Name, LabelToDisplay, Block.ToolTip, Block.Icon.Get());
 
 						if (bPopCommandList)
 						{
@@ -744,7 +869,7 @@ void UToolMenus::PopulateMenuBuilder(FMenuBuilder& MenuBuilder, UToolMenu* MenuD
 						}
 						else
 						{
-							MenuBuilder.AddMenuEntry(Block.Label, Block.ToolTip, Block.Icon.Get(), UIAction, Block.Name, Block.UserInterfaceActionType, Block.TutorialHighlightName);
+							MenuBuilder.AddMenuEntry(LabelToDisplay, Block.ToolTip, Block.Icon.Get(), UIAction, Block.Name, Block.UserInterfaceActionType, Block.TutorialHighlightName);
 						}
 					}
 				}
@@ -755,7 +880,14 @@ void UToolMenus::PopulateMenuBuilder(FMenuBuilder& MenuBuilder, UToolMenu* MenuD
 			}
 			else if (Block.Type == EMultiBlockType::Widget)
 			{
-				MenuBuilder.AddWidget(Widget.ToSharedRef(), Block.Label.Get(), Block.WidgetData.bNoIndent, Block.WidgetData.bSearchable);
+				if (bIsEditing)
+				{
+					MenuBuilder.AddMenuEntry(LabelToDisplay, Block.ToolTip, Block.Icon.Get(), UIAction, Block.Name, Block.UserInterfaceActionType, Block.TutorialHighlightName);
+				}
+				else
+				{
+					MenuBuilder.AddWidget(Widget.ToSharedRef(), LabelToDisplay.Get(), Block.WidgetData.bNoIndent, Block.WidgetData.bSearchable);
+				}
 			}
 			else
 			{
@@ -771,6 +903,27 @@ void UToolMenus::PopulateMenuBuilder(FMenuBuilder& MenuBuilder, UToolMenu* MenuD
 
 void UToolMenus::PopulateToolBarBuilder(FToolBarBuilder& ToolBarBuilder, UToolMenu* MenuData)
 {
+	if (GetEditMenusMode() && !MenuData->IsEditing() && EditMenuDelegate.IsBound())
+	{
+		TWeakObjectPtr<UToolMenu> WeakMenuPtr = MenuData;
+		const FName MenuName = MenuData->GetMenuName();
+		ToolBarBuilder.BeginSection(MenuName);
+		ToolBarBuilder.AddToolBarButton(
+			FExecuteAction::CreateLambda([MenuName, WeakMenuPtr]()
+			{
+				FPlatformApplicationMisc::ClipboardCopy(*MenuName.ToString());
+				if (UToolMenu* InMenu = WeakMenuPtr.Get())
+				{
+					UToolMenus::Get()->EditMenuDelegate.ExecuteIfBound(InMenu);
+				}
+			}), 
+			"MenuName",
+			LOCTEXT("EditMenu", "Edit Menu"),
+			LOCTEXT("EditMenu_ToolTip", "Open menu editor")
+		);
+		ToolBarBuilder.EndSection();
+	}
+
 	for (FToolMenuSection& Section : MenuData->Sections)
 	{
 		if (Section.Construct.NewToolBarDelegateLegacy.IsBound())
@@ -862,38 +1015,37 @@ void UToolMenus::PopulateToolBarBuilder(FToolBarBuilder& ToolBarBuilder, UToolMe
 		ToolBarBuilder.EndSection();
 	}
 
-	if (GetDisplayUIExtensionPoints())
-	{
-		const FName MenuName = MenuData->GetMenuName();
-		ToolBarBuilder.BeginSection(MenuName);
-		ToolBarBuilder.AddToolBarButton(
-			FExecuteAction::CreateLambda([MenuName]() { FPlatformApplicationMisc::ClipboardCopy(*MenuName.ToString()); }), 
-			"MenuName",
-			LOCTEXT("CopyNameToClipboard", "Copy Name"),
-			LOCTEXT("CopyMenuNameToClipboard", "Copy menu name to clipboard")
-		);
-		ToolBarBuilder.EndSection();
-	}
-
 	AddReferencedContextObjects(ToolBarBuilder.GetMultiBox(), MenuData->Context);
 }
 
-void UToolMenus::PopulateMenuBarBuilder(FMenuBarBuilder& MenuBarBuilder, UToolMenu* GeneratedMenu)
+void UToolMenus::PopulateMenuBarBuilder(FMenuBarBuilder& MenuBarBuilder, UToolMenu* MenuData)
 {
-	if (GeneratedMenu->Sections.Num() > 0)
+	for (int i=0; i < MenuData->Sections.Num(); ++i)
 	{
-		for (FToolMenuEntry& Block : GeneratedMenu->Sections[0].Blocks)
+		const FToolMenuSection& Section = MenuData->Sections[i];
+		for (const FToolMenuEntry& Block : Section.Blocks)
 		{
+			FName SubMenuFullName = JoinMenuPaths(MenuData->MenuName, Block.Name);
+			FNewMenuDelegate NewMenuDelegate;
+			if (Block.SubMenuData.ConstructMenu.NewMenuDelegate.IsBound())
+			{
+				NewMenuDelegate = Block.SubMenuData.ConstructMenu.NewMenuDelegate;
+			}
+			else
+			{
+				NewMenuDelegate = FNewMenuDelegate::CreateUObject(this, &UToolMenus::PopulateSubMenu, TWeakObjectPtr<UToolMenu>(MenuData), Block.Name);
+			}
+
 			MenuBarBuilder.AddPullDownMenu(
 				Block.Label.Get(),
 				Block.ToolTip.Get(),
-				FNewMenuDelegate::CreateUObject(this, &UToolMenus::FillMenuBarDropDown, GeneratedMenu->MenuName, Block.Name, GeneratedMenu->Context),
+				NewMenuDelegate,
 				Block.Name
 			);
 		}
-
-		AddReferencedContextObjects(MenuBarBuilder.GetMultiBox(), GeneratedMenu->Context);
 	}
+
+	AddReferencedContextObjects(MenuBarBuilder.GetMultiBox(), MenuData->Context);
 }
 
 FOnGetContent UToolMenus::ConvertWidgetChoice(const FNewToolMenuWidgetChoice& Choice, const FToolMenuContext& Context) const
@@ -1160,12 +1312,46 @@ void UToolMenus::AddReferencedObjects(UObject* InThis, FReferenceCollector& Coll
 	Super::AddReferencedObjects(InThis, Collector);
 }
 
+UToolMenu* UToolMenus::GenerateMenuOrSubMenuForEdit(const UToolMenu* InMenu)
+{
+	// Make copy of context so we can set bIsEditing flag on it
+	FToolMenuContext NewMenuContext = InMenu->Context;
+	NewMenuContext.bIsEditing = true;
+
+	if (!InMenu->SubMenuParent)
+	{
+		return GenerateMenu(InMenu->GetMenuName(), NewMenuContext);
+	}
+
+	// Generate each menu leading up to the final submenu because sub-menus are not required to be registered
+	TArray<const UToolMenu*> SubMenuChain = InMenu->GetSubMenuChain();
+	if (SubMenuChain.Num() > 0)
+	{
+		UToolMenu* CurrentGeneratedMenu = GenerateMenu(SubMenuChain[0]->GetMenuName(), NewMenuContext);
+		for (int32 i=1; i < SubMenuChain.Num(); ++i)
+		{
+			if (UToolMenu* Menu = GenerateSubMenu(CurrentGeneratedMenu, SubMenuChain[i]->SubMenuSourceEntryName))
+			{
+				CurrentGeneratedMenu = Menu;
+			}
+			else
+			{
+				return nullptr;
+			}
+		}
+
+		return CurrentGeneratedMenu;
+	}
+
+	return nullptr;
+}
+
 UToolMenu* UToolMenus::GenerateMenu(const FName Name, const FToolMenuContext& InMenuContext)
 {
-	return GenerateMenu(CollectHierarchy(Name), InMenuContext);
+	return GenerateMenuFromHierarchy(CollectHierarchy(Name), InMenuContext);
 }
 
-UToolMenu* UToolMenus::GenerateMenu(const TArray<UToolMenu*>& Hierarchy, const FToolMenuContext& InMenuContext)
+UToolMenu* UToolMenus::GenerateMenuFromHierarchy(const TArray<UToolMenu*>& Hierarchy, const FToolMenuContext& InMenuContext)
 {
 	UToolMenu* GeneratedMenu = NewObject<UToolMenu>(this);
 
@@ -1191,7 +1377,7 @@ TSharedRef<SWidget> UToolMenus::GenerateWidget(const TArray<UToolMenu*>& Hierarc
 		return SNullWidget::NullWidget;
 	}
 
-	UToolMenu* Generated = GenerateMenu(Hierarchy, InMenuContext);
+	UToolMenu* Generated = GenerateMenuFromHierarchy(Hierarchy, InMenuContext);
 	return GenerateWidget(Generated);
 }
 
@@ -1223,7 +1409,12 @@ TSharedRef<SWidget> UToolMenus::GenerateWidget(UToolMenu* GeneratedMenu)
 		}
 
 		FMenuBuilder MenuBuilder(GeneratedMenu->bShouldCloseWindowAfterMenuSelection, GeneratedMenu->Context.CommandList, GeneratedMenu->Context.GetAllExtenders(), GeneratedMenu->bCloseSelfOnly, GeneratedMenu->StyleSet, GeneratedMenu->bSearchable, GeneratedMenu->MenuName);
+		MenuBuilder.GetMultiBox()->WeakToolMenu = GeneratedMenu;
 		PopulateMenuBuilder(MenuBuilder, GeneratedMenu);
+		if (GeneratedMenu->ModifyBlockWidgetAfterMake.IsBound())
+		{
+			MenuBuilder.GetMultiBox()->ModifyBlockWidgetAfterMake = GeneratedMenu->ModifyBlockWidgetAfterMake;
+		}
 		TSharedRef<SWidget> Result = MenuBuilder.MakeWidget();
 		GeneratedMenuWidget.Widget = Result;
 		return Result;
@@ -1287,9 +1478,11 @@ void UToolMenus::SetNextTickTimer()
 {
 	if (!bNextTickTimerIsSet)
 	{
-		bNextTickTimerIsSet = true;
-
-		SetTimerForNextTickDelegate.ExecuteIfBound();
+		if (SetTimerForNextTickDelegate.IsBound())
+		{
+			bNextTickTimerIsSet = true;
+			SetTimerForNextTickDelegate.Execute();
+		}
 	}
 }
 
@@ -1426,7 +1619,7 @@ UToolMenu* UToolMenus::GenerateMenuAsBuilder(const UToolMenu* InMenu, const FToo
 		Hierarchy.Add((UToolMenu*)InMenu);
 	}
 
-	return GenerateMenu(Hierarchy, InMenuContext);
+	return GenerateMenuFromHierarchy(Hierarchy, InMenuContext);
 }
 
 UToolMenu* UToolMenus::RegisterMenu(const FName InName, const FName InParent, EMultiBoxType InType, bool bWarnIfAlreadyRegistered)
@@ -1583,3 +1776,9 @@ void UToolMenus::UnregisterStringCommandHandler(const FName InName)
 {
 	StringCommandHandlers.Remove(InName);
 }
+
+void UToolMenus::SaveCustomizations()
+{
+	SaveConfig();
+}
+
diff --git a/Engine/Source/Developer/ToolMenus/Public/ToolMenu.h b/Engine/Source/Developer/ToolMenus/Public/ToolMenu.h
--- a/Engine/Source/Developer/ToolMenus/Public/ToolMenu.h
+++ b/Engine/Source/Developer/ToolMenus/Public/ToolMenu.h
@@ -56,13 +56,25 @@ public:
 
 	bool IsRegistered() const { return bRegistered; }
 
+	/** returns array [Menu, Menu.SubMenuA, Menu.SubMenuB] for Menu.SubMenuB.SubMenuB */
+	TArray<const UToolMenu*> GetSubMenuChain() const;
+
+	/** returns "SubMenuC.SubMenuD" for menu "ModuleA.MenuB.SubMenuC.SubMenuD" */
+	FString GetSubMenuNamePath() const;
+
 	//~ Begin UToolMenuBase Interface
 	virtual bool IsEditing() const override;
 	virtual FName GetSectionName(const FName InEntryName) const override;
+	virtual bool ContainsSection(const FName InName) const override;
+	virtual bool ContainsEntry(const FName InName) const override;
 	virtual FCustomizedToolMenu* FindMenuCustomization() const override;
 	virtual FCustomizedToolMenu* AddMenuCustomization() const override;
+	virtual FCustomizedToolMenuHierarchy GetMenuCustomizationHierarchy() const override;
+	virtual void UpdateMenuCustomizationFromMultibox(const TSharedRef<const FMultiBox>& InMultiBox) override;
 	//~ End UToolMenuBase Interface
 
+	TArray<FName> GetMenuHierarchyNames(bool bIncludeSubMenuRoot) const;
+
 	template <typename TContextType>
 	TContextType* FindContext() const
 	{
@@ -81,6 +93,10 @@ private:
 
 	bool FindEntry(const FName EntryName, int32& SectionIndex, int32& EntryIndex) const;
 
+	FToolMenuEntry* FindEntry(const FName EntryName);
+
+	const FToolMenuEntry* FindEntry(const FName EntryName) const;
+
 	int32 IndexOfSection(const FName SectionName) const;
 
 	int32 FindInsertIndex(const FToolMenuSection& InSection) const;
@@ -116,7 +132,7 @@ public:
 
 	UPROPERTY(EditAnywhere, BlueprintReadWrite, Category = "ToolBar")
 	bool bToolBarForceSmallIcons;
-		
+
 	/** Prevent menu from being customized */
 	UPROPERTY(EditAnywhere, BlueprintReadWrite, Category = "ToolBar")
 	bool bPreventCustomization;
@@ -130,6 +146,14 @@ public:
 	UPROPERTY()
 	TArray<FToolMenuSection> Sections;
 
+	UPROPERTY()
+	const UToolMenu* SubMenuParent;
+
+	UPROPERTY()
+	FName SubMenuSourceEntryName;
+
+	FMultiBox::FOnModifyBlockWidgetAfterMake ModifyBlockWidgetAfterMake;
+
 private:
 
 	bool bRegistered;
diff --git a/Engine/Source/Developer/ToolMenus/Public/ToolMenuSection.h b/Engine/Source/Developer/ToolMenus/Public/ToolMenuSection.h
--- a/Engine/Source/Developer/ToolMenus/Public/ToolMenuSection.h
+++ b/Engine/Source/Developer/ToolMenus/Public/ToolMenuSection.h
@@ -53,6 +53,7 @@ public:
 	}
 
 	FToolMenuEntry* FindEntry(const FName InName);
+	const FToolMenuEntry* FindEntry(const FName InName) const;
 
 private:
 
diff --git a/Engine/Source/Developer/ToolMenus/Public/ToolMenus.h b/Engine/Source/Developer/ToolMenus/Public/ToolMenus.h
--- a/Engine/Source/Developer/ToolMenus/Public/ToolMenus.h
+++ b/Engine/Source/Developer/ToolMenus/Public/ToolMenus.h
@@ -52,7 +52,7 @@ struct FGeneratedToolMenuWidgets
 	TArray<FGeneratedToolMenuWidget> Instances;
 };
 
-UCLASS()
+UCLASS(config=EditorPerProjectUserSettings)
 class TOOLMENUS_API UToolMenus : public UObject
 {
 	GENERATED_BODY()
@@ -92,7 +92,7 @@ public:
 	 */
 	static bool IsToolMenuUIEnabled();
 
-	static void RegisterStartupCallback(const FSimpleMulticastDelegate::FDelegate& Delegate)
+	static FDelegateHandle RegisterStartupCallback(const FSimpleMulticastDelegate::FDelegate& Delegate)
 	{
 		if (IsToolMenuUIEnabled())
 		{
@@ -103,9 +103,11 @@ public:
 			else
 			{
 				// Wait until UToolMenus has been initialized
-				FCoreDelegates::OnPostEngineInit.Add(Delegate);
+				return FCoreDelegates::OnPostEngineInit.Add(Delegate);
 			}
 		}
+
+		return FDelegateHandle();
 	}
 
 	static void UnRegisterStartupCallback(const void* UserPointer)
@@ -113,6 +115,11 @@ public:
 		FCoreDelegates::OnPostEngineInit.RemoveAll(UserPointer);
 	}
 
+	static void UnRegisterStartupCallback(FDelegateHandle InHandle)
+	{
+		FCoreDelegates::OnPostEngineInit.Remove(InHandle);
+	}
+
 	/**
 	 * Registers a menu by name
 	 * @param	Parent	Optional name of a menu to layer on top of.
@@ -222,12 +229,12 @@ public:
 	/** Create a finalized menu that combines all parents used to generate a widget. Advanced special use cases only. */
 	UToolMenu* GenerateMenu(const FName Name, const FToolMenuContext& InMenuContext);
 
-	/** Create a finalized menu that combines given hierarchy array that will generate a widget. Advanced special use cases only. */
-	UToolMenu* GenerateMenu(const TArray<UToolMenu*>& Hierarchy, const FToolMenuContext& InMenuContext);
-
 	/** Create a finalized menu based on a custom crafted menu. Advanced special use cases only. */
 	UToolMenu* GenerateMenuAsBuilder(const UToolMenu* InMenu, const FToolMenuContext& InMenuContext);
 
+	/* Generate either a menu or submenu ready for editing */
+	UToolMenu* GenerateMenuOrSubMenuForEdit(const UToolMenu* InMenu);
+
 	/** Bake final menu including calls to construction delegates, sorting, and customization */
 	void AssembleMenuHierarchy(UToolMenu* GeneratedMenu, const TArray<UToolMenu*>& Hierarchy);
 
@@ -261,12 +268,35 @@ public:
 	/** Find index of customization settings for a menu */
 	int32 FindMenuCustomizationIndex(const FName InName);
 
+	/** Generates sub menu by entry name in the given generated menu parent */
+	UToolMenu* GenerateSubMenu(const UToolMenu* InGeneratedParent, const FName InBlockName);
+
+	/** When true, adds command to open edit menu dialog to each menu */
+	bool GetEditMenusMode() const;
+
+	/** Enables adding command to open edit menu dialog to each menu */
+	void SetEditMenusMode(bool bEnable);
+
 	/** Displaying extension points is for debugging menus */
 	DECLARE_DELEGATE_RetVal(bool, FShouldDisplayExtensionPoints);
 	FShouldDisplayExtensionPoints ShouldDisplayExtensionPoints;
 
+	/* Delegate that opens a menu editor */
+	DECLARE_DELEGATE_OneParam(FEditMenuDelegate, class UToolMenu*);
+	FEditMenuDelegate EditMenuDelegate;
+
+	/** Icon to display in menus for command to open menu editor */
+	FSlateIcon EditMenuIcon;
+
+	/** Join two paths together */
 	static FName JoinMenuPaths(const FName Base, const FName Child);
 
+	/** Break apart a menu path into components */
+	static bool SplitMenuPath(const FName MenuPath, FName& OutLeft, FName& OutRight);
+
+	/** Save menu customizations to ini files */
+	void SaveCustomizations();
+
 	friend struct FToolMenuOwnerScoped;
 	friend struct FToolMenuStringCommand;
 
@@ -276,10 +306,11 @@ public:
 	virtual bool IsDestructionThreadSafe() const override { return false; }
 	//~ End UObject Interface
 
-	FSlateIcon EditMenuIcon;
-
 private:
 
+	/** Create a finalized menu that combines given hierarchy array that will generate a widget. Advanced special use cases only. */
+	UToolMenu* GenerateMenuFromHierarchy(const TArray<UToolMenu*>& Hierarchy, const FToolMenuContext& InMenuContext);
+
 	/** Sets a timer to be called next engine tick so that multiple repeated actions can be combined together. */
 	void SetNextTickTimer();
 
@@ -300,7 +331,6 @@ private:
 
 	UToolMenu* FindSubMenuToGenerateWith(const FName InParentName, const FName InChildName);
 
-	void FillMenu(FMenuBuilder& MenuBuilder, FName InMenuName, FToolMenuContext InMenuContext);
 	void FillMenuBarDropDown(FMenuBuilder& MenuBuilder, FName InParentName, FName InChildName, FToolMenuContext InMenuContext);
 	void PopulateMenuBuilder(FMenuBuilder& MenuBuilder, UToolMenu* MenuData);
 	void PopulateMenuBarBuilder(FMenuBarBuilder& MenuBarBuilder, UToolMenu* MenuData);
@@ -319,7 +349,7 @@ private:
 
 	static void ExecuteStringCommand(const FToolMenuStringCommand StringCommand, const FToolMenuContext Context);
 
-	void FillMenuDynamic(FMenuBuilder& Builder, const FName SubMenuFullName, FNewToolMenuDelegate InConstructMenu, const FToolMenuContext Context);
+	void PopulateSubMenu(FMenuBuilder& Builder, TWeakObjectPtr<UToolMenu> InParent, const FName InBlockName);
 
 	void ListAllParents(const FName Name, TArray<FName>& AllParents);
 
@@ -340,7 +370,7 @@ private:
 
 private:
 
-	UPROPERTY(EditAnywhere, Category = Misc)
+	UPROPERTY(config, EditAnywhere, Category = Misc)
 	TArray<FCustomizedToolMenu> CustomizedMenus;
 
 	UPROPERTY()
@@ -360,8 +390,8 @@ private:
 	bool bNextTickTimerIsSet;
 	bool bRefreshWidgetsNextTick;
 	bool bCleanupStaleWidgetsNextTick;
+	bool bEditMenusMode;
 
-	
 	static UToolMenus* Singleton;
 	static bool bHasShutDown;
 };
diff --git a/Engine/Source/Editor/EditorStyle/Private/SlateEditorStyle.cpp b/Engine/Source/Editor/EditorStyle/Private/SlateEditorStyle.cpp
--- a/Engine/Source/Editor/EditorStyle/Private/SlateEditorStyle.cpp
+++ b/Engine/Source/Editor/EditorStyle/Private/SlateEditorStyle.cpp
@@ -2016,6 +2016,17 @@ void FSlateEditorStyle::FStyle::SetupGeneralStyles()
 #endif // WITH_EDITOR
 
 #if WITH_EDITOR || (IS_PROGRAM && WITH_UNREAL_DEVELOPER_TOOLS)
+	// Menu editor
+	{
+		Set("MultiBox.DragAbove", new BOX_BRUSH("Common/DropZoneIndicator_Above", FMargin(10.0f / 16.0f, 10.0f / 16.0f, 0, 0), SelectionColor));
+		Set("MultiBox.DragBelow", new BOX_BRUSH("Common/DropZoneIndicator_Below", FMargin(10.0f / 16.0f, 0, 0, 10.0f / 16.0f), SelectionColor));
+
+		Set("MultiBox.VisibleIcon16x", new IMAGE_BRUSH("Icons/icon_visible_16px", Icon16x16));
+		Set("MultiBox.VisibleHighlightIcon16x", new IMAGE_BRUSH("Icons/icon_visible_hi_16px", Icon16x16));
+		Set("MultiBox.NotVisibleIcon16x", new IMAGE_BRUSH("Icons/icon_invisible_16px", Icon16x16));
+		Set("MultiBox.NotVisibleHighlightIcon16x", new IMAGE_BRUSH("Icons/icon_invisible_hi_16px", Icon16x16));
+	}
+
 	// ToolBar
 	{
 		Set( "ToolBar.Background", new BOX_BRUSH( "Common/GroupBorder", FMargin(4.0f/16.0f) ) );
diff --git a/Engine/Source/Editor/MainFrame/Private/Menus/MainMenu.cpp b/Engine/Source/Editor/MainFrame/Private/Menus/MainMenu.cpp
--- a/Engine/Source/Editor/MainFrame/Private/Menus/MainMenu.cpp
+++ b/Engine/Source/Editor/MainFrame/Private/Menus/MainMenu.cpp
@@ -119,37 +119,6 @@ void FMainMenu::RegisterEditMenu()
 
 	{
 		FToolMenuSection& Section = EditMenu->AddSection("EditLocalTabSpawners", LOCTEXT("ConfigurationHeading", "Configuration"));
-		if (GetDefault<UEditorExperimentalSettings>()->bToolbarCustomization)
-		{
-			FUIAction ToggleMultiBoxEditMode(
-				FExecuteAction::CreateStatic(&FMultiBoxSettings::ToggleToolbarEditing),
-				FCanExecuteAction(),
-				FIsActionChecked::CreateStatic(&FMultiBoxSettings::IsInToolbarEditMode)
-			);
-		
-			Section.AddMenuEntry(
-				"EditToolbars",
-				LOCTEXT("EditToolbarsLabel", "Edit Toolbars"),
-				LOCTEXT("EditToolbarsToolTip", "Allows customization of each toolbar"),
-				FSlateIcon(),
-				ToggleMultiBoxEditMode,
-				EUserInterfaceActionType::ToggleButton
-			);
-
-			Section.AddDynamicEntry("TabManager", FNewToolMenuDelegateLegacy::CreateLambda([](FMenuBuilder& MenuBuilder, UToolMenu* MenuData)
-			{
-				if (USlateTabManagerContext* TabManagerContext = MenuData->FindContext<USlateTabManagerContext>())
-				{
-					if (TabManagerContext->TabManager.IsValid())
-					{
-						// Automatically populate tab spawners from TabManager
-						const IWorkspaceMenuStructure& MenuStructure = WorkspaceMenu::GetMenuStructure();
-						TabManagerContext->TabManager.Pin()->PopulateTabSpawnerMenu(MenuBuilder, MenuStructure.GetEditOptions());
-					}
-				}
-			}));
-		}
-
 		if (GetDefault<UEditorStyleSettings>()->bExpandConfigurationMenus)
 		{
 			Section.AddSubMenu(
diff --git a/Engine/Source/Editor/ToolMenusEditor/Private/SEditToolMenuDialog.cpp b/Engine/Source/Editor/ToolMenusEditor/Private/SEditToolMenuDialog.cpp
new file mode 100644
--- /dev/null
+++ b/Engine/Source/Editor/ToolMenusEditor/Private/SEditToolMenuDialog.cpp
@@ -0,0 +1,911 @@
+// Copyright 1998-2019 Epic Games, Inc. All Rights Reserved.
+
+
+#include "SEditToolMenuDialog.h"
+#include "IToolMenusEditorModule.h"
+
+#include "Misc/MessageDialog.h"
+#include "HAL/FileManager.h"
+#include "Misc/App.h"
+#include "SlateOptMacros.h"
+#include "Widgets/Layout/SSeparator.h"
+#include "Widgets/Images/SImage.h"
+#include "Widgets/Layout/SGridPanel.h"
+#include "Widgets/Layout/SUniformGridPanel.h"
+#include "Widgets/Input/SEditableTextBox.h"
+#include "Widgets/Input/SButton.h"
+#include "Widgets/Input/SCheckBox.h"
+#include "Widgets/Layout/SScrollBox.h"
+
+#include "Framework/Docking/TabManager.h"
+#include "EditorStyleSet.h"
+
+#include "ToolMenus.h"
+
+#include "Editor.h"
+
+#include "IAssetRegistry.h"
+#include "AssetRegistryModule.h"
+
+#include "Serialization/StaticMemoryReader.h"
+#include "Widgets/Colors/SColorBlock.h"
+#include "Framework/Commands/UICommandDragDropOp.h"
+
+#define LOCTEXT_NAMESPACE "EditToolMenuDialog"
+
+class SVisibilityEditWidget : public SImage
+{
+public:
+	SLATE_BEGIN_ARGS(SVisibilityEditWidget) {}
+	SLATE_END_ARGS()
+
+	/** Construct this widget */
+	void Construct(const FArguments& InArgs, TSharedRef<SMultiBoxWidget> InBaseWidget, TSharedRef<const FMultiBlock> InBlock, TSharedRef<SEditToolMenuDialog> InDialogWidget)
+	{	
+		BaseWidget = InBaseWidget;
+		Block = InBlock;
+		DialogWidget = InDialogWidget;
+
+		SImage::Construct(
+			SImage::FArguments()
+			.Image(this, &SVisibilityEditWidget::GetBrush)
+		);
+	}
+
+private:
+
+	FReply HandleClick()
+	{
+		ToggleVisible();
+		return FReply::Handled();
+	}
+
+	virtual FReply OnMouseButtonDoubleClick(const FGeometry& InMyGeometry, const FPointerEvent& InMouseEvent) override
+	{
+		return HandleClick();
+	}
+
+	virtual FReply OnMouseButtonDown(const FGeometry& MyGeometry, const FPointerEvent& MouseEvent) override
+	{
+		if (MouseEvent.GetEffectingButton() != EKeys::LeftMouseButton)
+		{
+			return FReply::Unhandled();
+		}
+
+		return HandleClick();
+	}
+
+	virtual FReply OnMouseButtonUp( const FGeometry& MyGeometry, const FPointerEvent& MouseEvent ) override
+	{
+		if (MouseEvent.GetEffectingButton() == EKeys::LeftMouseButton)
+		{
+			return FReply::Handled();
+		}
+
+		return FReply::Unhandled();
+	}
+
+	const FSlateBrush* GetBrush() const
+	{
+		if (IsVisible())
+		{
+			static const FName NAME_VisibleHoveredBrush("MultiBox.VisibleHighlightIcon16x");
+			static const FName NAME_VisibleNotHoveredBrush("MultiBox.VisibleIcon16x");
+			return IsHovered() ? FEditorStyle::GetBrush(NAME_VisibleHoveredBrush) :
+				FEditorStyle::GetBrush(NAME_VisibleNotHoveredBrush);
+		}
+		else
+		{
+			static const FName NAME_NotVisibleHoveredBrush("MultiBox.NotVisibleHighlightIcon16x");
+			static const FName NAME_NotVisibleNotHoveredBrush("MultiBox.NotVisibleIcon16x");
+			return IsHovered() ? FEditorStyle::GetBrush(NAME_NotVisibleHoveredBrush) :
+				FEditorStyle::GetBrush(NAME_NotVisibleNotHoveredBrush);
+		}
+	}
+
+	bool IsVisible() const
+	{
+		const FName ItemName = Block->GetExtensionHook();
+		if (ItemName == NAME_None)
+		{
+			return true;
+		}
+
+		UToolMenuBase* ToolMenu = BaseWidget.IsValid() ? BaseWidget.Pin()->GetMultiBox()->GetToolMenu() : nullptr;
+		if (!ToolMenu)
+		{
+			return true;
+		}
+
+		FCustomizedToolMenuHierarchy MenuCustomizationHierarchy = ToolMenu->GetMenuCustomizationHierarchy();
+		if (MenuCustomizationHierarchy.Hierarchy.Num() > 0)
+		{
+			if (Block->GetType() == EMultiBlockType::Heading)
+			{
+				return !MenuCustomizationHierarchy.IsSectionHidden(ItemName);
+			}
+			else
+			{
+				return !MenuCustomizationHierarchy.IsEntryHidden(ItemName);
+			}
+		}
+
+		return true;
+	}	
+	
+	void ToggleVisible()
+	{
+		if (DialogWidget.IsValid())
+		{
+			DialogWidget.Pin()->OnToggleVisibleClicked(Block.ToSharedRef(), BaseWidget);
+		}
+	}
+
+	TSharedPtr<const FMultiBlock> Block;
+	TWeakPtr<SMultiBoxWidget> BaseWidget;
+	TWeakPtr<SEditToolMenuDialog> DialogWidget;
+};
+
+class SMultiBlockDragHandle : public SCompoundWidget
+{
+	SLATE_BEGIN_ARGS( SMultiBlockDragHandle ) {}
+	SLATE_END_ARGS()
+
+	void Construct( const FArguments& InArgs, TSharedRef<SMultiBoxWidget> InBaseWidget, TSharedRef<const FMultiBlock> InBlock, bool bInIsBlockNamedAndToolMenu );
+
+	FReply OnMouseButtonDown( const FGeometry& MyGeometry, const FPointerEvent& MouseEvent ) override;
+	FReply OnMouseButtonUp( const FGeometry& MyGeometry, const FPointerEvent& MouseEvent ) override;
+
+	virtual void OnDragEnter( const FGeometry& MyGeometry, const FDragDropEvent& DragDropEvent ) override;
+	virtual FReply OnDragOver( const FGeometry& MyGeometry, const FDragDropEvent& DragDropEvent ) override;
+	virtual FReply OnDrop( const FGeometry& MyGeometry, const FDragDropEvent& DragDropEvent ) override;
+	virtual FReply OnDragDetected( const FGeometry& MyGeometry, const FPointerEvent& MouseEvent ) override;
+
+	virtual int32 OnPaint(const FPaintArgs& Args, const FGeometry& AllottedGeometry, const FSlateRect& MyCullingRect, FSlateWindowElementList& OutDrawElements, int32 LayerId, const FWidgetStyle& InWidgetStyle, bool bParentEnabled) const override
+	{
+		LayerId = SCompoundWidget::OnPaint(Args, AllottedGeometry, MyCullingRect, OutDrawElements, LayerId, InWidgetStyle, bParentEnabled);
+	
+		bool bIsDropDestination = false;
+		bool bInsertAfter = false;
+		if (BaseWidget.IsValid() && Block.IsValid())
+		{
+			const EVisibility Status = BaseWidget.Pin()->GetCustomizationBorderDragVisibility(Block->GetExtensionHook(), Block->GetType(), bInsertAfter);
+			if (Status == EVisibility::Visible)
+			{
+				bIsDropDestination = true;
+			}
+		}
+
+		if (bIsDropDestination)
+		{
+			const FSlateBrush* DropIndicatorBrush = FEditorStyle::GetBrush(bInsertAfter ? "MultiBox.DragBelow" : "MultiBox.DragAbove");
+
+			FSlateDrawElement::MakeBox
+			(
+				OutDrawElements,
+				LayerId++,
+				AllottedGeometry.ToPaintGeometry(),
+				DropIndicatorBrush,
+				ESlateDrawEffect::None,
+				DropIndicatorBrush->GetTint(InWidgetStyle) * InWidgetStyle.GetColorAndOpacityTint()
+			);
+		}
+
+		return LayerId;
+	}
+
+private:
+
+	TSharedPtr<const FMultiBlock> Block;
+	TWeakPtr<SMultiBoxWidget> BaseWidget;
+	bool bIsBlockNamedAndToolMenu;
+};
+
+
+void SMultiBlockDragHandle::Construct( const FArguments& InArgs, TSharedRef<SMultiBoxWidget> InBaseWidget, TSharedRef<const FMultiBlock> InBlock, bool bInIsBlockNamedAndToolMenu )
+{
+	BaseWidget = InBaseWidget;
+	Block = InBlock;
+	bIsBlockNamedAndToolMenu = bInIsBlockNamedAndToolMenu;
+}
+
+FReply SMultiBlockDragHandle::OnMouseButtonDown( const FGeometry& MyGeometry, const FPointerEvent& MouseEvent )
+{
+	if (MouseEvent.GetEffectingButton() == EKeys::LeftMouseButton)
+	{
+		if (Block->GetExtensionHook() != NAME_None && bIsBlockNamedAndToolMenu)
+		{
+			return FReply::Handled().DetectDrag(SharedThis(this), MouseEvent.GetEffectingButton());
+		}
+	}
+
+	return FReply::Unhandled();
+}
+
+FReply SMultiBlockDragHandle::OnMouseButtonUp( const FGeometry& MyGeometry, const FPointerEvent& MouseEvent )
+{
+	if (MouseEvent.GetEffectingButton() == EKeys::LeftMouseButton)
+	{
+		BaseWidget.Pin()->GetMultiBox()->OnEditSelectionChanged().ExecuteIfBound(Block.ToSharedRef());
+	}
+
+	return FReply::Unhandled();
+}
+
+void SMultiBlockDragHandle::OnDragEnter( const FGeometry& MyGeometry, const FDragDropEvent& DragDropEvent )
+{
+	if ( DragDropEvent.GetOperationAs<FUICommandDragDropOp>().IsValid() )
+	{
+		BaseWidget.Pin()->OnCustomCommandDragEnter( Block.ToSharedRef(), MyGeometry, DragDropEvent );
+	}
+}
+
+FReply SMultiBlockDragHandle::OnDragOver( const FGeometry& MyGeometry, const FDragDropEvent& DragDropEvent )
+{
+	if ( DragDropEvent.GetOperationAs<FUICommandDragDropOp>().IsValid() )
+	{
+		BaseWidget.Pin()->OnCustomCommandDragged( Block.ToSharedRef(), MyGeometry, DragDropEvent );
+		return FReply::Handled();
+	}
+
+	return FReply::Unhandled();
+}
+
+FReply SMultiBlockDragHandle::OnDrop( const FGeometry& MyGeometry, const FDragDropEvent& DragDropEvent )
+{
+	if ( DragDropEvent.GetOperationAs<FUICommandDragDropOp>().IsValid() )
+	{
+		BaseWidget.Pin()->OnCustomCommandDropped();
+		return FReply::Handled();
+	}
+
+	return FReply::Unhandled();
+}
+
+FReply SMultiBlockDragHandle::OnDragDetected( const FGeometry& MyGeometry, const FPointerEvent& MouseEvent )
+{
+	TSharedPtr<SWidget> CustomDecorator;
+
+	TSharedRef<SWidget> BlockWidget = Block->MakeWidget(BaseWidget.Pin().ToSharedRef(), EMultiBlockLocation::None, false)->AsWidget();
+
+	if (Block->IsSeparator())
+	{
+		CustomDecorator = SNew(SVerticalBox)
+
+		+SVerticalBox::Slot()
+		.AutoHeight()
+		[
+			SNew(SHorizontalBox)
+			+SHorizontalBox::Slot()
+			[
+				SNew(SSeparator)
+			]
+		]
+
+		+SVerticalBox::Slot()
+		.AutoHeight()
+		[
+			SNew(SHorizontalBox)
+		
+			+SHorizontalBox::Slot()
+			.Padding( 20, 0, 40, 0 )
+			.AutoWidth()
+			[
+				SNew(STextBlock)
+				.Text(FText::FromName(Block->GetExtensionHook()))
+			]
+		];
+	}
+	else if (Block->GetType() == EMultiBlockType::Heading)
+	{
+		CustomDecorator = SNew(SVerticalBox)
+
+		+SVerticalBox::Slot()
+		.AutoHeight()
+		[
+			SNew(SHorizontalBox)
+			+SHorizontalBox::Slot()
+			[
+				SNew(SSeparator)
+			]
+		]
+
+		+SVerticalBox::Slot()
+		.AutoHeight()
+		[
+			SNew(SHorizontalBox)
+		
+			+SHorizontalBox::Slot()
+			.Padding( 20, 0, 40, 0 )
+			.AutoWidth()
+			[
+				BlockWidget
+			]
+		];
+	}
+	else
+	{
+		CustomDecorator = BlockWidget;
+	}
+
+	TSharedRef<FUICommandDragDropOp> NewOp = FUICommandDragDropOp::New(
+			Block->GetExtensionHook(),
+			Block->GetType(),
+			Block->IsPartOfHeading(),
+			NAME_None, 
+			CustomDecorator,
+			MyGeometry.AbsolutePosition-MouseEvent.GetScreenSpacePosition()
+		);
+
+	NewOp->SetOnDropNotification( FSimpleDelegate::CreateSP( BaseWidget.Pin().ToSharedRef(), &SMultiBoxWidget::OnDropExternal ) );
+
+	TSharedRef<FDragDropOperation> DragDropOp = NewOp;
+	return FReply::Handled().BeginDragDrop( DragDropOp );
+}
+
+void SEditToolMenuDialog::InitMenu(UToolMenu* InMenu)
+{
+	MenuDialogOpenedWith = InMenu;
+
+	OriginalSettings.Reset();
+	if (UToolMenu* GeneratedMenu = UToolMenus::Get()->GenerateMenuOrSubMenuForEdit(InMenu))
+	{
+		for (const FCustomizedToolMenu* CustomizedToolMenu : GeneratedMenu->GetMenuCustomizationHierarchy().Hierarchy)
+		{
+			if (CustomizedToolMenu)
+			{
+				OriginalSettings.Add(*CustomizedToolMenu);
+			}
+		}
+		CurrentGeneratedMenu = GeneratedMenu;
+		MenuNames = GeneratedMenu->GetMenuHierarchyNames(true);
+	}
+	else
+	{
+		CurrentGeneratedMenu = nullptr;
+		MenuNames.Reset();
+	}
+}
+
+void SEditToolMenuDialog::Construct( const FArguments& InArgs )
+{
+	bClosedWithOk = false;
+
+	InitMenu(InArgs._SourceMenu);
+
+	SetSelectedItem(NAME_None, ESelectedEditMenuEntryType::Menu);
+	BuildWidget();
+}
+
+void SEditToolMenuDialog::HandleOnLiveCodingPatchComplete()
+{
+	Refresh();
+}
+
+void SEditToolMenuDialog::OnSelectedEntryChanged(TSharedRef<const FMultiBlock> InBlock)
+{
+	UToolMenu* ToolMenu = CurrentGeneratedMenu.Get();
+	const FName BlockName = InBlock->GetExtensionHook();
+
+	ESelectedEditMenuEntryType SelectedType = ESelectedEditMenuEntryType::None;
+	if (InBlock->IsPartOfHeading())
+	{
+		SelectedType = ESelectedEditMenuEntryType::Section;
+		if (!ToolMenu->ContainsSection(BlockName))
+		{
+			return;
+		}
+	}
+	else
+	{
+		SelectedType = ESelectedEditMenuEntryType::Entry;
+		if (!ToolMenu->ContainsEntry(BlockName))
+		{
+			return;
+		}
+	}
+
+	SetSelectedItem(BlockName, SelectedType);
+	if (PropertiesWidget.IsValid())
+	{
+		PropertiesWidget->SetObject(SelectedObject.Get());
+	}
+}
+
+void SEditToolMenuDialog::SetSelectedItem(const FName InName, ESelectedEditMenuEntryType InType)
+{
+	UToolMenu* ToolMenu = CurrentGeneratedMenu.Get();
+	if (!ToolMenu)
+	{
+		SelectedObject.Reset();
+		return;
+	}
+
+	if (ESelectedEditMenuEntryType::Entry == InType)
+	{
+		UToolMenuEditorDialogEntry* NewSelectedObject = NewObject<UToolMenuEditorDialogEntry>();
+		NewSelectedObject->Init(ToolMenu, InName);
+		SelectedObject = TStrongObjectPtr<UToolMenuEditorDialogEntry>(NewSelectedObject);
+	}
+	else if (ESelectedEditMenuEntryType::Section == InType)
+	{
+		UToolMenuEditorDialogSection* NewSelectedObject = NewObject<UToolMenuEditorDialogSection>();
+		NewSelectedObject->Init(ToolMenu, InName);
+		SelectedObject = TStrongObjectPtr<UToolMenuEditorDialogSection>(NewSelectedObject);
+	}
+	else if (ESelectedEditMenuEntryType::Menu == InType)
+	{
+		UToolMenuEditorDialogMenu* NewSelectedObject = NewObject<UToolMenuEditorDialogMenu>();
+		NewSelectedObject->Init(ToolMenu, ToolMenu->GetMenuName());
+		SelectedObject = TStrongObjectPtr<UToolMenuEditorDialogMenu>(NewSelectedObject);
+	}
+	else
+	{
+		SelectedObject.Reset();
+	}
+}
+
+TSharedRef<SWidget> SEditToolMenuDialog::ModifyBlockWidgetAfterMake(const TSharedRef<SMultiBoxWidget>& InMultiBoxWidget, const FMultiBlock& InBlock, const TSharedRef<SWidget>& InBlockWidget)
+{
+	TSharedRef<const FMultiBox> MultiBox = InMultiBoxWidget->GetMultiBox();
+
+	const bool bIsEditing = MultiBox->IsInEditMode();
+
+	static const FName SelectionColor("SelectionColor");
+
+	UToolMenuBase* ToolMenuBase = MultiBox->GetToolMenu();
+	const FName BlockName = InBlock.GetExtensionHook();
+	const EMultiBlockType BlockType = InBlock.GetType();
+
+	bool bIsBlockNamedAndToolMenu = false;
+	if (BlockName != NAME_None && ToolMenuBase)
+	{
+		if (InBlock.IsPartOfHeading())
+		{
+			bIsBlockNamedAndToolMenu = ToolMenuBase->ContainsSection(BlockName);
+		}
+		else
+		{
+			bIsBlockNamedAndToolMenu = ToolMenuBase->ContainsEntry(BlockName);
+		}
+	}
+
+	if (!bIsBlockNamedAndToolMenu)
+	{
+		return SNew(SOverlay)
+
+		+SOverlay::Slot()
+		[
+			SNew(SHorizontalBox)
+			
+			+SHorizontalBox::Slot()
+			.Padding(36, 0, 0, 0)
+			[
+				InBlockWidget
+			]
+		]
+			
+		+SOverlay::Slot()
+		[
+			SNew( SMultiBlockDragHandle, InMultiBoxWidget, InBlock.AsShared(), bIsBlockNamedAndToolMenu )
+		];
+	}
+	else if (InBlock.IsSeparator() && InBlock.IsPartOfHeading())
+	{
+		return SNew(SOverlay)
+
+		+SOverlay::Slot()
+		[
+			InBlockWidget
+		]
+
+		+SOverlay::Slot()
+		[
+			SNew( SMultiBlockDragHandle, InMultiBoxWidget, InBlock.AsShared(), bIsBlockNamedAndToolMenu )
+		];
+	}
+	else
+	{
+		return SNew(SOverlay)
+
+		+SOverlay::Slot()
+		[
+			SNew(SHorizontalBox)
+
+			// Drag handle
+			+SHorizontalBox::Slot()
+			.AutoWidth()
+			.VAlign(VAlign_Center)
+			[
+				SNew(SOverlay)
+				+SOverlay::Slot()
+				[
+					SNew(SVerticalBox)
+					+SVerticalBox::Slot()
+					.Padding(4, 0, 32, 0)
+					.AutoHeight()
+					[
+						SNew(SImage)
+						.Image(FCoreStyle::Get().GetBrush("VerticalBoxDragIndicatorShort"))
+					]
+				]
+			]
+
+			// Widget content
+			+SHorizontalBox::Slot()
+			[
+				SNew(SOverlay)
+				+SOverlay::Slot()
+				[
+					InBlockWidget
+				]
+			]
+		]
+
+		// Prevent input from reaching BlockWidget and allow drag/drop
+		+SOverlay::Slot()
+		[
+			SNew( SMultiBlockDragHandle, InMultiBoxWidget, InBlock.AsShared(), bIsBlockNamedAndToolMenu )
+		]
+
+		// Toggle Hidden
+		+SOverlay::Slot()
+		[
+			SNew(SHorizontalBox)
+			+SHorizontalBox::Slot()
+			.AutoWidth()
+			.VAlign(VAlign_Center)
+			[
+				SNew(SVerticalBox)
+				+SVerticalBox::Slot()
+				.Padding(20, 0, 0, 0)
+				.AutoHeight()
+				[
+					// TODO: need to be able to update the proxy object being displayed in details panel
+					// or better, perhaps have the proxy object displayed in the details panel directly lookup the config from UToolMenus singleton?
+					SNew(SVisibilityEditWidget, InMultiBoxWidget, InBlock.AsShared(), SharedThis(this))
+				]
+			]
+		];
+	}
+}
+
+void SEditToolMenuDialog::BuildWidget()
+{
+	UToolMenu* ToolMenu = CurrentGeneratedMenu.Get();
+	if (!ToolMenu)
+	{
+		ChildSlot
+		[
+			SNew(STextBlock)
+			.TextStyle(FEditorStyle::Get(), "LargeText")
+			.Text(LOCTEXT("Unavailable", "Unavailable"))
+		];
+
+		return;
+	}
+	
+	ToolMenu->ModifyBlockWidgetAfterMake.BindSP(this, &SEditToolMenuDialog::ModifyBlockWidgetAfterMake);
+	TSharedRef<SWidget> MenuWidget = UToolMenus::Get()->GenerateWidget(ToolMenu);
+
+	TSharedRef<SMultiBoxWidget> MultiBoxWidget = StaticCastSharedRef<SMultiBoxWidget>(MenuWidget);
+	TSharedRef<FMultiBox> MultiBox = ConstCastSharedRef<FMultiBox>(MultiBoxWidget->GetMultiBox());
+	MultiBox->OnEditSelectionChanged().BindSP(this, &SEditToolMenuDialog::OnSelectedEntryChanged);
+
+	TSharedRef<SWidget> DialogCenter =
+	
+		SNew(SHorizontalBox)
+		+SHorizontalBox::Slot()
+		.AutoWidth()
+		[
+			// Preview
+			SNew(SScrollBox)
+
+			+SScrollBox::Slot()
+			[
+				SNew(SVerticalBox)
+
+				+SVerticalBox::Slot()
+				.AutoHeight()
+				.Padding(0, 8, 0, 8)
+				[
+					MenuWidget
+				]
+			]
+		]
+
+		+SHorizontalBox::Slot()
+		.Padding(32, 0, 0, 0)
+		[
+			SNew(SVerticalBox)
+
+			// Dialog buttons Save/Cancel
+			+SVerticalBox::Slot()
+			.AutoHeight()
+			.Padding(0)
+			.HAlign(HAlign_Right)
+			[
+				SNew(SHorizontalBox)
+								
+				+SHorizontalBox::Slot()
+				.AutoWidth()
+				[
+					SNew(SButton)
+					.Text(LOCTEXT("UndoAllChangesButtonText", "Undo all changes"))
+					.ToolTipText(LOCTEXT("UndoAllChangesButtonTooltip", "Undos all changes made to these menus"))
+					.OnClicked(this, &SEditToolMenuDialog::UndoAllChanges)
+				]
+
+				+SHorizontalBox::Slot()
+				.Padding(8, 0, 0, 0)
+				.AutoWidth()
+				[
+					SNew(SButton)
+					.Text(LOCTEXT("ResetDefaultsButtonText", "Reset to Defaults"))
+					.ToolTipText(LOCTEXT("ResetDefaultsButtonTooltip", "Reset settings for current menu to default values"))
+					.OnClicked(this, &SEditToolMenuDialog::HandleResetClicked)
+				]
+			]
+
+			+SVerticalBox::Slot()
+			.AutoHeight()
+			[
+				BuildMenuPropertiesWidget()
+			]
+		];
+
+	{
+		MenuNameComboData.Reset();
+		for (FName It : MenuNames)
+		{
+			MenuNameComboData.Add(MakeShareable(new FName(It)));
+		}
+	}
+
+	ChildSlot
+	[
+		SNew(SBorder)
+		.Padding(20)
+		.BorderImage( FEditorStyle::GetBrush("Docking.Tab.Normal") )
+		[
+			SNew(SVerticalBox)
+
+			// Title
+			+SVerticalBox::Slot()
+			.AutoHeight()
+			.Padding(0)
+			[
+				SNew(STextBlock)
+				.TextStyle(FEditorStyle::Get(), "LargeText")
+				.Text(FText::FromName(ToolMenu->MenuName))
+			]
+
+			// Menu names
+			+SVerticalBox::Slot()
+			.AutoHeight()
+			.Padding(0)
+			[
+				SNew(SComboBox< TSharedPtr<FName> >)
+				.OptionsSource(&MenuNameComboData)
+				.OnGenerateWidget(this, &SEditToolMenuDialog::MakeMenuNameComboEntryWidget)
+				.OnSelectionChanged(this, &SEditToolMenuDialog::OnMenuNamesSelectionChanged)
+				[
+					SNew(STextBlock)
+					.Text(FText::FromName(ToolMenu->MenuName))
+				]
+			]
+
+			// Title spacer
+			+SVerticalBox::Slot()
+			.AutoHeight()
+			.Padding(0, 16, 0, 8)
+			[
+				SNew(SSeparator)
+			]
+
+			// Middle of dialog
+			+SVerticalBox::Slot()
+			[
+				DialogCenter
+			]
+		]
+	];
+}
+
+TSharedRef<SWidget> SEditToolMenuDialog::BuildMenuPropertiesWidget()
+{
+	FDetailsViewArgs Args;
+	Args.bHideSelectionTip = true;
+
+	FPropertyEditorModule& PropertyModule = FModuleManager::LoadModuleChecked<FPropertyEditorModule>("PropertyEditor");
+	PropertiesWidget = PropertyModule.CreateDetailView(Args);
+	PropertiesWidget->SetObject(SelectedObject.Get());
+
+	return PropertiesWidget.ToSharedRef();
+}
+
+TSharedRef<SWidget> SEditToolMenuDialog::MakeMenuNameComboEntryWidget(TSharedPtr<FName> InEntry)
+{
+	return SNew(STextBlock)
+			.Text(FText::FromName(*InEntry));
+}
+
+void SEditToolMenuDialog::OnMenuNamesSelectionChanged(TSharedPtr<FName> InEntry, ESelectInfo::Type SelectInfo)
+{
+	UToolMenu* SourceMenu = MenuDialogOpenedWith.Get();
+	if (!SourceMenu)
+	{
+		BuildWidget();
+		return;
+	}
+
+	if (!InEntry.IsValid())
+	{
+		BuildWidget();
+		return;
+	}
+
+	const FName* NewMenuNamePtr = InEntry.Get();
+	const FName NewMenuName = NewMenuNamePtr ? *NewMenuNamePtr : NAME_None;
+
+	FName OldMenuName;
+	if (UToolMenu* ToolMenu = CurrentGeneratedMenu.Get())
+	{
+		if (ToolMenu->GetMenuName() == NewMenuName)
+		{
+			return;
+		}
+
+		OldMenuName = ToolMenu->GetMenuName();
+	}
+	
+	FToolMenuContext NewMenuContext = SourceMenu->Context;
+	NewMenuContext.bIsEditing = true;
+		
+	TArray<const UToolMenu*> SubMenuChain = SourceMenu->GetSubMenuChain();
+	if (!SourceMenu->SubMenuParent)
+	{
+		CurrentGeneratedMenu = UToolMenus::Get()->GenerateMenu(NewMenuName, NewMenuContext);
+	}
+	else if (SubMenuChain.Num() > 0)
+	{
+		FString SubMenuNamePath = SourceMenu->GetSubMenuNamePath();
+
+		FString NewMenuNameString = NewMenuName.ToString();
+		int32 NumCharactersToKeep = NewMenuNameString.Len() - SubMenuNamePath.Len() - 1;
+		if (NumCharactersToKeep > 0)
+		{
+			// "Menu" and "SubMenuA"
+			FString BaseMenuPath = NewMenuNameString.Left(NumCharactersToKeep);
+
+			// Create the base menu, then walk the down the SubMenus
+			if (UToolMenu* NewMenu = UToolMenus::Get()->GenerateMenu(*BaseMenuPath, NewMenuContext))
+			{
+				for (int32 i=1; i < SubMenuChain.Num(); ++i)
+				{
+					if (UToolMenu* SubMenu = UToolMenus::Get()->GenerateSubMenu(NewMenu, SubMenuChain[i]->SubMenuSourceEntryName))
+					{
+						NewMenu = SubMenu;
+					}
+					else
+					{
+						NewMenu = nullptr;
+						break;
+					}
+				}
+
+				if (NewMenu)
+				{
+					CurrentGeneratedMenu = NewMenu;
+				}
+			}
+		}
+	}
+
+	BuildWidget();
+}
+
+void SEditToolMenuDialog::OnToggleVisibleClicked(TSharedRef<const FMultiBlock> InBlock, TWeakPtr<SMultiBoxWidget> BaseWidget)
+{
+	const FName ItemName = InBlock->GetExtensionHook();
+	if (ItemName == NAME_None)
+	{
+		return;
+	}
+
+	UToolMenuBase* ToolMenu = BaseWidget.IsValid() ? BaseWidget.Pin()->GetMultiBox()->GetToolMenu() : nullptr;
+	if (!ToolMenu)
+	{
+		return;
+	}
+
+	// Commit change
+	FCustomizedToolMenu* MenuCustomization = ToolMenu->AddMenuCustomization();
+	FCustomizedToolMenuHierarchy MenuCustomizationHierarchy = ToolMenu->GetMenuCustomizationHierarchy();
+	if (InBlock->GetType() == EMultiBlockType::Heading)
+	{
+		bool bHidden = MenuCustomizationHierarchy.IsSectionHidden(ItemName);
+		MenuCustomization->AddSection(ItemName)->Visibility = bHidden ? ECustomizedToolMenuVisibility::Visible : ECustomizedToolMenuVisibility::Hidden;
+	}
+	else
+	{
+		bool bHidden = MenuCustomizationHierarchy.IsEntryHidden(ItemName);
+		MenuCustomization->AddEntry(ItemName)->Visibility = bHidden ? ECustomizedToolMenuVisibility::Visible : ECustomizedToolMenuVisibility::Hidden;
+	}
+
+	// Update details panel
+	if (SelectedObject.IsValid())
+	{
+		SelectedObject->LoadState();
+		SelectedObject->PostEditChange();
+		if (PropertiesWidget.IsValid())
+		{
+			PropertiesWidget->ForceRefresh();
+		}
+	}
+}
+
+void SEditToolMenuDialog::CloseContainingWindow()
+{
+	TSharedPtr<SWindow> ContainingWindow = FSlateApplication::Get().FindWidgetWindow(AsShared());
+
+	if ( ContainingWindow.IsValid() )
+	{
+		ContainingWindow->RequestDestroyWindow();
+	}
+}
+
+FReply SEditToolMenuDialog::Refresh()
+{
+	if (UToolMenu* ToolMenu = CurrentGeneratedMenu.Get())
+	{
+		CurrentGeneratedMenu = UToolMenus::Get()->GenerateMenuOrSubMenuForEdit(ToolMenu);
+		UToolMenus::Get()->RefreshMenuWidget(ToolMenu->GetMenuName());
+	}
+
+	BuildWidget();
+	
+	return FReply::Handled();
+}
+
+FReply SEditToolMenuDialog::HandleResetClicked()
+{
+	if (UToolMenu* ToolMenu = CurrentGeneratedMenu.Get())
+	{
+		UToolMenus::Get()->RemoveCustomization(ToolMenu->GetMenuName());
+		Refresh();
+	}
+
+	return FReply::Handled();
+}
+
+void SEditToolMenuDialog::SaveSettingsToDisk()
+{
+	UToolMenus::Get()->SaveCustomizations();
+}
+
+FReply SEditToolMenuDialog::UndoAllChanges()
+{
+	UToolMenus* ToolMenus = UToolMenus::Get();
+	for (const FName HierarchyMenuName : MenuNames)
+	{
+		ToolMenus->RemoveCustomization(HierarchyMenuName);
+	}
+
+	for (const FCustomizedToolMenu& Original : OriginalSettings)
+	{
+		FCustomizedToolMenu* Destination = ToolMenus->AddMenuCustomization(Original.Name);
+		*Destination = Original;
+	}
+
+	Refresh();
+	return FReply::Handled();
+}
+
+void SEditToolMenuDialog::OnWindowClosed(const TSharedRef<SWindow>& Window)
+{	
+	SaveSettingsToDisk();
+	UToolMenus::Get()->RefreshAllWidgets();
+}
+
+#undef LOCTEXT_NAMESPACE
diff --git a/Engine/Source/Editor/ToolMenusEditor/Private/SEditToolMenuDialog.h b/Engine/Source/Editor/ToolMenusEditor/Private/SEditToolMenuDialog.h
new file mode 100644
--- /dev/null
+++ b/Engine/Source/Editor/ToolMenusEditor/Private/SEditToolMenuDialog.h
@@ -0,0 +1,84 @@
+// Copyright 1998-2019 Epic Games, Inc. All Rights Reserved.
+
+#pragma once
+
+#include "CoreMinimal.h"
+#include "Layout/Visibility.h"
+#include "Widgets/DeclarativeSyntaxSupport.h"
+#include "Styling/SlateColor.h"
+#include "Input/Reply.h"
+#include "Widgets/SWidget.h"
+#include "Widgets/SCompoundWidget.h"
+#include "Widgets/Views/STableViewBase.h"
+#include "Widgets/Views/STableRow.h"
+#include "Widgets/Input/SComboBox.h"
+#include "ToolMenus.h"
+#include "UObject/StrongObjectPtr.h"
+#include "ToolMenusEditor.h"
+
+
+class UToolMenuEditorDialogEntry;
+
+
+/**
+ * A dialog to customize a menu
+ */
+class SEditToolMenuDialog : public SCompoundWidget
+{
+public:
+	SLATE_BEGIN_ARGS( SEditToolMenuDialog )
+	{}
+
+	/** A reference to the parent window */
+	SLATE_ARGUMENT(TSharedPtr<SWindow>, ParentWindow)
+
+	SLATE_ARGUMENT(UToolMenu*, SourceMenu)
+
+	SLATE_END_ARGS()
+
+	/** Constructs this widget with InArgs */
+	void Construct( const FArguments& InArgs );
+	
+	void OnToggleVisibleClicked(TSharedRef<const FMultiBlock> InBlock, TWeakPtr<SMultiBoxWidget> BaseWidget);
+	
+	/** Closes the window that contains this widget */
+	void CloseContainingWindow();
+
+	void OnWindowClosed(const TSharedRef<SWindow>& Window);
+
+private:
+
+	TSharedRef<SWidget> ModifyBlockWidgetAfterMake(const TSharedRef<SMultiBoxWidget>& InMultiBoxWidget, const FMultiBlock& InBlock, const TSharedRef<SWidget>& InBlockWidget);
+
+	void BuildWidget();
+	void HandleOnLiveCodingPatchComplete();
+
+	void InitMenu(UToolMenu* InMenu);
+
+	FReply Refresh();
+	FReply HandleResetClicked();
+	FReply UndoAllChanges();
+
+	void SaveSettingsToDisk();
+
+	TSharedRef<SWidget> BuildMenuPropertiesWidget();
+	TSharedRef<SWidget> MakeMenuNameComboEntryWidget(TSharedPtr<FName> InEntry);
+	void OnMenuNamesSelectionChanged(TSharedPtr<FName> InEntry, ESelectInfo::Type SelectInfo);
+
+	void OnSelectedEntryChanged(TSharedRef<const FMultiBlock> InBlock);
+
+	void SetSelectedItem(const FName InName, ESelectedEditMenuEntryType InType);
+
+	bool bClosedWithOk;
+
+	TWeakObjectPtr<UToolMenu> MenuDialogOpenedWith;
+	TWeakObjectPtr<UToolMenu> CurrentGeneratedMenu;
+
+	TArray<FName> MenuNames;
+	TArray<TSharedPtr<FName>> MenuNameComboData;
+
+	TSharedPtr<class IDetailsView> PropertiesWidget;
+	TStrongObjectPtr<UToolMenuEditorDialogObject> SelectedObject;
+
+	TArray<FCustomizedToolMenu> OriginalSettings;
+};
diff --git a/Engine/Source/Editor/ToolMenusEditor/Private/ToolMenusEditor.cpp b/Engine/Source/Editor/ToolMenusEditor/Private/ToolMenusEditor.cpp
new file mode 100644
--- /dev/null
+++ b/Engine/Source/Editor/ToolMenusEditor/Private/ToolMenusEditor.cpp
@@ -0,0 +1,96 @@
+// Copyright 1998-2019 Epic Games, Inc. All Rights Reserved.
+
+#include "ToolMenusEditor.h"
+#include "ToolMenus.h"
+
+void UToolMenuEditorDialogMenu::Init(UToolMenu* InMenu, const FName InName)
+{
+	Menu = InMenu;
+	Name = InName;
+
+	LoadState();
+}
+
+void UToolMenuEditorDialogSection::Init(UToolMenu* InMenu, const FName InName)
+{
+	Name = InName;
+	Type = ESelectedEditMenuEntryType::Section;
+	Menu = InMenu;
+
+	LoadState();
+}
+
+void UToolMenuEditorDialogSection::LoadState()
+{
+	Super::LoadState();
+
+	Visibility = ECustomizedToolMenuVisibility::None;
+
+	if (Menu && Name != NAME_None)
+	{
+		if (FCustomizedToolMenu* CustomizedToolMenu = Menu->FindMenuCustomization())
+		{
+			if (FCustomizedToolMenuSection* Found = CustomizedToolMenu->Sections.Find(Name))
+			{
+				Visibility = Found->Visibility;
+			}
+		}
+	}
+}
+
+void UToolMenuEditorDialogSection::PostEditChangeProperty(FPropertyChangedEvent& PropertyChangedEvent)
+{
+	Super::PostEditChangeProperty(PropertyChangedEvent);
+
+	const FName PropertyName = PropertyChangedEvent.GetPropertyName();
+	if (PropertyName == GET_MEMBER_NAME_CHECKED(UToolMenuEditorDialogEntry, Visibility))
+	{
+		if (Menu)
+		{
+			FCustomizedToolMenu* CustomizedToolMenu = Menu->AddMenuCustomization();
+			CustomizedToolMenu->AddSection(Name)->Visibility = Visibility;
+		}
+	}
+}
+
+void UToolMenuEditorDialogEntry::Init(UToolMenu* InMenu, const FName InName)
+{
+	Name = InName;
+	Type = ESelectedEditMenuEntryType::Entry;
+	Menu = InMenu;
+
+	LoadState();
+}
+
+void UToolMenuEditorDialogEntry::LoadState()
+{
+	Super::LoadState();
+
+	Visibility = ECustomizedToolMenuVisibility::None;
+
+	if (Menu && Name != NAME_None)
+	{
+		if (FCustomizedToolMenu* CustomizedToolMenu = Menu->FindMenuCustomization())
+		{
+			if (FCustomizedToolMenuEntry* Found = CustomizedToolMenu->Entries.Find(Name))
+			{
+				Visibility = Found->Visibility;
+			}
+		}
+	}
+}
+
+void UToolMenuEditorDialogEntry::PostEditChangeProperty(FPropertyChangedEvent& PropertyChangedEvent)
+{
+	Super::PostEditChangeProperty(PropertyChangedEvent);
+
+	const FName PropertyName = PropertyChangedEvent.GetPropertyName();
+	if (PropertyName == GET_MEMBER_NAME_CHECKED(UToolMenuEditorDialogEntry, Visibility))
+	{
+		if (Menu)
+		{
+			FCustomizedToolMenu* CustomizedToolMenu = Menu->AddMenuCustomization();
+			CustomizedToolMenu->AddEntry(Name)->Visibility = Visibility;
+		}
+	}
+}
diff --git a/Engine/Source/Editor/ToolMenusEditor/Private/ToolMenusEditorModule.cpp b/Engine/Source/Editor/ToolMenusEditor/Private/ToolMenusEditorModule.cpp
new file mode 100644
--- /dev/null
+++ b/Engine/Source/Editor/ToolMenusEditor/Private/ToolMenusEditorModule.cpp
@@ -0,0 +1,101 @@
+// Copyright 1998-2019 Epic Games, Inc. All Rights Reserved.
+
+#include "ToolMenusEditor.h"
+
+#include "CoreMinimal.h"
+#include "Modules/ModuleManager.h"
+#include "IToolMenusEditorModule.h"
+#include "SEditToolMenuDialog.h"
+#include "ToolMenus.h"
+
+#define LOCTEXT_NAMESPACE "ToolMenusEditor"
+
+static bool ToolMenusEditWindowModalCVar = true;
+FAutoConsoleVariableRef CVarRex(
+	TEXT("ToolMenus.EditWindowModal"),
+	ToolMenusEditWindowModalCVar,
+	TEXT("Edit menus in modal window"),
+	ECVF_Default);
+
+/**
+ * Implements the Tool menus module.
+ */
+class FToolMenusEditorModule
+	: public IToolMenusEditorModule
+{
+public:
+
+	// IModuleInterface interface
+
+	virtual void StartupModule() override
+	{
+	}
+
+	virtual void ShutdownModule() override
+	{
+	}
+
+	// End IModuleInterface interface
+
+	virtual void RegisterShowEditMenusModeCheckbox() const override
+	{
+		UToolMenu* Menu = UToolMenus::Get()->ExtendMenu("LevelEditor.MainMenu.Window");
+		FToolMenuSection& Section = Menu->FindOrAddSection("WindowLayout");
+
+		Section.AddMenuEntry(
+			"EnableMenuEditing",
+			LOCTEXT("EnableMenuEditing", "Enable Menu Editing"),
+			LOCTEXT("EnableMenuEditing_ToolTip", "Adds command to each menu and toolbar for editing"),
+			FSlateIcon(),
+			FUIAction(
+				FExecuteAction::CreateLambda([]()
+				{
+					UToolMenus* ToolMenus = UToolMenus::Get();
+					ToolMenus->SetEditMenusMode(!ToolMenus->GetEditMenusMode());
+				}),
+				FCanExecuteAction(),
+				FGetActionCheckState::CreateLambda([]()
+				{
+					return UToolMenus::Get()->GetEditMenusMode() ? ECheckBoxState::Checked : ECheckBoxState::Unchecked; 
+				})),
+			EUserInterfaceActionType::ToggleButton
+		);
+	}
+
+	virtual void OpenEditToolMenuDialog(UToolMenu* ToolMenu) const
+	{
+		TSharedPtr<SWindow> ParentWindow = FGlobalTabmanager::Get()->GetRootWindow();
+		const FVector2D WindowSize = FVector2D(940, 540);
+		FText WindowTitle = LOCTEXT("EditMenu_Title", "Edit Menu");
+
+		TSharedRef<SWindow> EditMenuWindow =
+			SNew(SWindow)
+			.Title(WindowTitle)
+			.ClientSize(WindowSize)
+			.SupportsMinimize(true)
+			.SupportsMaximize(true);
+
+		TSharedRef<SEditToolMenuDialog> NewDialog =
+			SNew(SEditToolMenuDialog)
+			.ParentWindow(EditMenuWindow)
+			.SourceMenu(ToolMenu);
+
+		EditMenuWindow->SetContent(NewDialog);
+
+		EditMenuWindow->SetOnWindowClosed(FOnWindowClosed::CreateSP(NewDialog, &SEditToolMenuDialog::OnWindowClosed));
+
+		if (ToolMenusEditWindowModalCVar)
+		{
+			FSlateApplication::Get().AddModalWindow(EditMenuWindow, ParentWindow);
+		}
+		else
+		{
+			FSlateApplication::Get().AddWindow(EditMenuWindow);
+		}
+	}
+};
+
+#undef LOCTEXT_NAMESPACE
+
+IMPLEMENT_MODULE(FToolMenusEditorModule, ToolMenusEditor)
+
diff --git a/Engine/Source/Editor/ToolMenusEditor/Public/IToolMenusEditorModule.h b/Engine/Source/Editor/ToolMenusEditor/Public/IToolMenusEditorModule.h
new file mode 100644
--- /dev/null
+++ b/Engine/Source/Editor/ToolMenusEditor/Public/IToolMenusEditorModule.h
@@ -0,0 +1,44 @@
+// Copyright 1998-2019 Epic Games, Inc. All Rights Reserved.
+#pragma once
+
+#include "CoreMinimal.h"
+#include "Modules/ModuleInterface.h"
+#include "Modules/ModuleManager.h"
+
+DEFINE_LOG_CATEGORY_STATIC(LogToolMenusEditor, Log, All);
+
+class IToolMenusEditorModule : public IModuleInterface
+{
+public:
+
+	/**
+	 * Retrieve the module instance.
+	 */
+	static inline IToolMenusEditorModule& Get()
+	{
+		return FModuleManager::LoadModuleChecked<IToolMenusEditorModule>("ToolMenusEditor");
+	}
+
+	/**
+	 * Checks to see if this module is loaded and ready.  It is only valid to call Get() if IsAvailable() returns true.
+	 *
+	 * @return True if the module is loaded and ready to use
+	 */
+	static inline bool IsAvailable()
+	{
+		return FModuleManager::Get().IsModuleLoaded("ToolMenusEditor");
+	}
+
+	/**
+	 * Open editor for the given ToolMenu
+	 *
+	 * @param ToolMenu The menu top edit
+	 */
+	virtual void OpenEditToolMenuDialog(class UToolMenu* ToolMenu) const = 0;
+
+	/**
+	 * Adds menu entry to toggle adding an entry to each menu and toolbar to open a menu editor
+	 */
+	virtual void RegisterShowEditMenusModeCheckbox() const = 0;
+};
+
diff --git a/Engine/Source/Editor/ToolMenusEditor/Public/ToolMenusEditor.h b/Engine/Source/Editor/ToolMenusEditor/Public/ToolMenusEditor.h
new file mode 100644
--- /dev/null
+++ b/Engine/Source/Editor/ToolMenusEditor/Public/ToolMenusEditor.h
@@ -0,0 +1,93 @@
+// Copyright 1998-2019 Epic Games, Inc. All Rights Reserved.
+
+#pragma once
+
+#include "CoreMinimal.h"
+#include "UObject/ObjectMacros.h"
+#include "UObject/UObjectGlobals.h"
+#include "ToolMenus.h"
+
+#include "ToolMenusEditor.generated.h"
+
+UENUM()
+enum class ESelectedEditMenuEntryType : uint8
+{
+	None = 0,
+	Entry,
+	Section,
+	Menu
+};
+
+UCLASS(Abstract)
+class UToolMenuEditorDialogObject : public UObject
+{
+	GENERATED_BODY()
+
+public:
+
+	virtual void LoadState() {}
+};
+
+UCLASS()
+class UToolMenuEditorDialogMenu : public UToolMenuEditorDialogObject
+{
+	GENERATED_BODY()
+
+public:
+
+	void Init(class UToolMenu* InMenu, const FName InName);
+
+	UPROPERTY(VisibleAnywhere, Category = Misc)
+	FName Name;
+
+	UPROPERTY()
+	UToolMenu* Menu;
+};
+
+UCLASS()
+class UToolMenuEditorDialogBlock : public UToolMenuEditorDialogObject
+{
+	GENERATED_BODY()
+
+public:
+
+	UPROPERTY(VisibleAnywhere, Category = Misc)
+	FName Name;
+
+	UPROPERTY(VisibleAnywhere, Category = Misc)
+	ESelectedEditMenuEntryType Type;
+
+	UPROPERTY(EditAnywhere, Category = Misc)
+	ECustomizedToolMenuVisibility Visibility;
+
+	UPROPERTY()
+	UToolMenu* Menu;
+};
+
+UCLASS()
+class UToolMenuEditorDialogSection : public UToolMenuEditorDialogBlock
+{
+	GENERATED_BODY()
+
+public:
+
+	void Init(class UToolMenu* InMenu, const FName InName);
+
+	virtual void LoadState() override;
+
+	virtual void PostEditChangeProperty(struct FPropertyChangedEvent& PropertyChangedEvent) override;
+};
+
+UCLASS()
+class UToolMenuEditorDialogEntry : public UToolMenuEditorDialogBlock
+{
+	GENERATED_BODY()
+
+public:
+
+	void Init(class UToolMenu* InMenu, const FName InName);
+
+	virtual void LoadState() override;
+
+	virtual void PostEditChangeProperty(struct FPropertyChangedEvent& PropertyChangedEvent) override;
+};
diff --git a/Engine/Source/Editor/ToolMenusEditor/ToolMenusEditor.Build.cs b/Engine/Source/Editor/ToolMenusEditor/ToolMenusEditor.Build.cs
new file mode 100644
--- /dev/null
+++ b/Engine/Source/Editor/ToolMenusEditor/ToolMenusEditor.Build.cs
@@ -0,0 +1,33 @@
+// Copyright 1998-2019 Epic Games, Inc. All Rights Reserved.
+
+namespace UnrealBuildTool.Rules
+{
+	public class ToolMenusEditor : ModuleRules
+	{
+		public ToolMenusEditor(ReadOnlyTargetRules Target) : base(Target)
+		{
+			PrivateDependencyModuleNames.AddRange(
+				new string[] {
+					"Core",
+					"CoreUObject",
+					"Slate",
+					"SlateCore",
+					"ToolMenus",
+					"EditorStyle",
+					"UnrealEd",
+					"InputCore",
+				}
+			);
+
+			PrivateIncludePathModuleNames.AddRange(
+				new string[] {
+					"AssetRegistry"
+			});
+
+			DynamicallyLoadedModuleNames.AddRange(
+				new string[] {
+					"AssetRegistry"
+				});
+		}
+	}
+}
diff --git a/Engine/Source/Editor/UnrealEd/Private/EditorEngine.cpp b/Engine/Source/Editor/UnrealEd/Private/EditorEngine.cpp
--- a/Engine/Source/Editor/UnrealEd/Private/EditorEngine.cpp
+++ b/Engine/Source/Editor/UnrealEd/Private/EditorEngine.cpp
@@ -214,6 +214,7 @@
 #include "RenderTargetPool.h"
 #include "RenderGraphBuilder.h"
 #include "ToolMenus.h"
+#include "IToolMenusEditorModule.h"
 #include "Subsystems/AssetEditorSubsystem.h"
 
 
@@ -785,7 +786,9 @@ void UEditorEngine::InitEditor(IEngineLoop* InEngineLoop)
 
 	if (FSlateApplication::IsInitialized() && UToolMenus::IsToolMenuUIEnabled())
 	{
-		TWeakPtr<FTimerManager> WeakTimerManager;
+		UToolMenus::Get()->EditMenuIcon = FSlateIcon(FEditorStyle::GetStyleSetName(), "ContentBrowser.AssetActions.OpenInExternalEditor");
+
+		TWeakPtr<FTimerManager> WeakTimerManager = TimerManager;
 		UToolMenus::Get()->AssignSetTimerForNextTickDelegate(FSimpleDelegate::CreateLambda([WeakTimerManager]()
 		{
 			if (WeakTimerManager.IsValid())
@@ -794,6 +797,18 @@ void UEditorEngine::InitEditor(IEngineLoop* InEngineLoop)
 			}
 		}));
 
+		bool bShowEditMenusModeCheckbox = false;
+		GConfig->GetBool(TEXT("/Script/UnrealEd.EditorExperimentalSettings"), TEXT("bShowEditMenusModeCheckbox"), bShowEditMenusModeCheckbox, GEditorPerProjectIni);
+		if (bShowEditMenusModeCheckbox)
+		{
+			IToolMenusEditorModule::Get().RegisterShowEditMenusModeCheckbox();
+
+			UToolMenus::Get()->EditMenuDelegate.BindLambda([](UToolMenu* InMenu)
+			{
+				IToolMenusEditorModule::Get().OpenEditToolMenuDialog(InMenu);
+			});
+		}
+
 		UToolMenus::Get()->ShouldDisplayExtensionPoints.BindStatic(&GetDisplayMultiboxHooks);
 
 		UToolMenus::Get()->RegisterStringCommandHandler("Command", FToolMenuExecuteString::CreateLambda([](const FString& InString, const FToolMenuContext& InContext)
diff --git a/Engine/Source/Editor/UnrealEd/UnrealEd.Build.cs b/Engine/Source/Editor/UnrealEd/UnrealEd.Build.cs
--- a/Engine/Source/Editor/UnrealEd/UnrealEd.Build.cs
+++ b/Engine/Source/Editor/UnrealEd/UnrealEd.Build.cs
@@ -102,7 +102,8 @@ public class UnrealEd : ModuleRules
                 "MeshBuilder",
                 "MaterialShaderQualitySettings",
                 "EditorSubsystem",
-                "InteractiveToolsFramework"
+                "InteractiveToolsFramework",
+				"ToolMenusEditor",
             }
 		);
 
diff --git a/Engine/Source/Runtime/Slate/Private/Framework/Commands/UICommandDragDropOp.cpp b/Engine/Source/Runtime/Slate/Private/Framework/Commands/UICommandDragDropOp.cpp
--- a/Engine/Source/Runtime/Slate/Private/Framework/Commands/UICommandDragDropOp.cpp
+++ b/Engine/Source/Runtime/Slate/Private/Framework/Commands/UICommandDragDropOp.cpp
@@ -8,9 +8,9 @@
 #include "Widgets/Text/STextBlock.h"
 
 
-TSharedRef<FUICommandDragDropOp> FUICommandDragDropOp::New( TSharedRef<const FUICommandInfo> InCommandInfo, FName InOriginMultiBox, TSharedPtr<SWidget> CustomDectorator, FVector2D DecoratorOffset )
+TSharedRef<FUICommandDragDropOp> FUICommandDragDropOp::New( FName InItemName, EMultiBlockType InBlockType, bool bInIsDraggingSection, FName InOriginMultiBox, TSharedPtr<SWidget> CustomDectorator, FVector2D DecoratorOffset )
 {
-	TSharedRef<FUICommandDragDropOp> Operation = MakeShareable(new FUICommandDragDropOp(InCommandInfo, InOriginMultiBox, CustomDectorator, DecoratorOffset));
+	TSharedRef<FUICommandDragDropOp> Operation = MakeShareable(new FUICommandDragDropOp(InItemName, InBlockType, bInIsDraggingSection, InOriginMultiBox, CustomDectorator, DecoratorOffset));
 	Operation->Construct();
 
 	return Operation;
@@ -42,7 +42,7 @@ TSharedPtr<SWidget> FUICommandDragDropOp::GetDefaultDecorator() const
 	else
 	{
 		Content = SNew(STextBlock)
-			.Text(UICommand->GetLabel());
+			.Text(FText::FromName(ItemName));
 	}
 
 	// Create hover widget
diff --git a/Engine/Source/Runtime/Slate/Private/Framework/MultiBox/MultiBox.cpp b/Engine/Source/Runtime/Slate/Private/Framework/MultiBox/MultiBox.cpp
--- a/Engine/Source/Runtime/Slate/Private/Framework/MultiBox/MultiBox.cpp
+++ b/Engine/Source/Runtime/Slate/Private/Framework/MultiBox/MultiBox.cpp
@@ -8,7 +8,9 @@
 #include "Framework/Application/SlateApplication.h"
 #include "Widgets/Layout/SBorder.h"
 #include "Framework/MultiBox/MultiBoxBuilder.h"
+#include "Framework/MultiBox/ToolMenuBase.h"
 #include "Widgets/Input/SButton.h"
+#include "Widgets/Images/SImage.h"
 #include "Widgets/SToolTip.h"
 #include "Widgets/Views/STableViewBase.h"
 #include "Widgets/Views/STableRow.h"
@@ -16,6 +18,7 @@
 #include "Widgets/Layout/SScrollBox.h"
 #include "Framework/MultiBox/SToolBarButtonBlock.h"
 #include "Framework/MultiBox/SMenuEntryBlock.h"
+#include "Framework/MultiBox/SMenuSeparatorBlock.h"
 #include "Framework/MultiBox/MultiBoxCustomization.h"
 #include "Framework/MultiBox/SClippingHorizontalBox.h"
 #include "Framework/Commands/UICommandDragDropOp.h"
@@ -25,8 +28,6 @@ TAttribute<bool> FMultiBoxSettings::UseSmallToolBarIcons;
 TAttribute<bool> FMultiBoxSettings::DisplayMultiboxHooks;
 FMultiBoxSettings::FConstructToolTip FMultiBoxSettings::ToolTipConstructor = FConstructToolTip::CreateStatic( &FMultiBoxSettings::ConstructDefaultToolTip );
 
-bool FMultiBoxSettings::bInToolbarEditMode;
-
 
 FMultiBoxSettings::FMultiBoxSettings()
 {
@@ -51,17 +52,6 @@ void FMultiBoxSettings::ResetToolTipConstructor()
 	ToolTipConstructor = FConstructToolTip::CreateStatic( &FMultiBoxSettings::ConstructDefaultToolTip );
 }
 
-void FMultiBoxSettings::ToggleToolbarEditing()
-{
-	bool bCanBeEnabled = false;
-	if( GIsEditor )
-	{
-		GConfig->GetBool(TEXT("/Script/UnrealEd.EditorExperimentalSettings"), TEXT("bToolbarCustomization"), bCanBeEnabled, GEditorPerProjectIni);
-	}
-
-	bInToolbarEditMode = !bInToolbarEditMode && bCanBeEnabled;
-}
-
 const FMultiBoxCustomization FMultiBoxCustomization::None( NAME_None );
 
 
@@ -126,6 +116,16 @@ FReply SMultiBlockBaseWidget::OnDrop( const FGeometry& MyGeometry, const FDragDr
 	return FReply::Unhandled();
 }
 
+bool SMultiBlockBaseWidget::IsInEditMode() const
+{
+	if (OwnerMultiBoxWidget.IsValid())
+	{
+		return OwnerMultiBoxWidget.Pin()->GetMultiBox()->IsInEditMode();
+	}
+
+	return false;
+}
+
 /**
  * Creates a MultiBlock widget for this MultiBlock
  *
@@ -173,8 +173,7 @@ bool FMultiBlock::GetSearchable() const
  * @param	bInShouldCloseWindowAfterMenuSelection	Sets whether or not the window that contains this multibox should be destroyed after the user clicks on a menu item in this box
  */
 FMultiBox::FMultiBox( const EMultiBoxType InType, FMultiBoxCustomization InCustomization, const bool bInShouldCloseWindowAfterMenuSelection )
-	: CustomizationData( new FMultiBoxCustomizationData( InCustomization.GetCustomizationName() ) )
-	, CommandLists()
+	: CommandLists()
 	, Blocks()
 	, StyleSet( &FCoreStyle::Get() )
 	, StyleName( "ToolBar" )
@@ -221,9 +220,6 @@ void FMultiBox::RemoveCustomMultiBlock( TSharedRef< const FMultiBlock> InBlock )
 		if( Index != INDEX_NONE )
 		{
 			Blocks.RemoveAt( Index );
-
-			// Remove the block from the customization data
-			CustomizationData->BlockRemoved( InBlock, Index, Blocks );
 		}
 
 	}
@@ -231,24 +227,103 @@ void FMultiBox::RemoveCustomMultiBlock( TSharedRef< const FMultiBlock> InBlock )
 
 void FMultiBox::InsertCustomMultiBlock( TSharedRef<const FMultiBlock> InBlock, int32 Index )
 {
-	if( IsCustomizable() && ensure( InBlock->GetAction().IsValid() ) )
+	if (IsCustomizable() && ensure(InBlock->GetExtensionHook() != NAME_None))
 	{
 		int32 ExistingIndex = Blocks.Find( InBlock );
-		if( ExistingIndex != INDEX_NONE )
+
+		FName DestinationBlockName = NAME_None;
+		FName DestinationSectionName = NAME_None;
+		if (Blocks.IsValidIndex(Index))
 		{
-			Blocks.RemoveAt( ExistingIndex );
+			DestinationBlockName = Blocks[Index]->GetExtensionHook();
 
-			CustomizationData->BlockRemoved( InBlock, ExistingIndex, Blocks );
+			int32 DestinationSectionEndIndex = INDEX_NONE;
+			int32 DestinationSectionIndex = GetSectionEditBounds(Index, DestinationSectionEndIndex);
+			if (Blocks.IsValidIndex(DestinationSectionIndex))
+			{
+				DestinationSectionName = Blocks[DestinationSectionIndex]->GetExtensionHook();
+			}
+		}
 
-			if( ExistingIndex < Index )
+		if (InBlock->IsPartOfHeading())
+		{
+			if (InBlock->GetExtensionHook() == DestinationSectionName)
 			{
-				--Index;
+				return;
+			}
+
+			if (ExistingIndex != INDEX_NONE)
+			{
+				int32 SourceSectionEndIndex = INDEX_NONE;
+				int32 SourceSectionIndex = GetSectionEditBounds(ExistingIndex, SourceSectionEndIndex);
+				if (SourceSectionIndex != INDEX_NONE && SourceSectionEndIndex != INDEX_NONE)
+				{
+					bool bHadSeparator = Blocks[SourceSectionIndex]->IsSeparator();
+
+					TArray< TSharedRef< const FMultiBlock > > BlocksToMove;
+					BlocksToMove.Reset(SourceSectionEndIndex - SourceSectionIndex + 1);
+					for (int32 BlockIdx = SourceSectionIndex; BlockIdx < SourceSectionEndIndex; ++BlockIdx)
+					{
+						BlocksToMove.Add(Blocks[BlockIdx]);
+					}
+
+					Blocks.RemoveAt(SourceSectionIndex, SourceSectionEndIndex - SourceSectionIndex, false);
+
+					if (Index > SourceSectionIndex)
+					{
+						Index -= BlocksToMove.Num();
+					}
+
+					if (Index == 0)
+					{
+						// Add missing separator for next section
+						if (Blocks.Num() > 0 && (Blocks[0]->GetType() == EMultiBlockType::Heading))
+						{
+							BlocksToMove.Add(MakeShareable(new FMenuSeparatorBlock(Blocks[0]->GetExtensionHook(), /* bInIsPartOfHeading=*/ true)));
+						}
+					}
+					else
+					{
+						// Add separator to beginning of section
+						if (BlocksToMove.Num() > 0 && (BlocksToMove[0]->GetType() == EMultiBlockType::Heading))
+						{
+							BlocksToMove.Insert(MakeShareable(new FMenuSeparatorBlock(BlocksToMove[0]->GetExtensionHook(), /* bInIsPartOfHeading=*/ true)), 0);
+						}
+					}
+
+					Blocks.Insert(BlocksToMove, Index);
+
+					// Menus do not start with separators, remove separator if one exists
+					if (Blocks.Num() > 0 && Blocks[0]->IsSeparator())
+					{
+						Blocks.RemoveAt(0, 1, false);
+					}
+
+					if (UToolMenuBase* ToolMenu = GetToolMenu())
+					{
+						ToolMenu->UpdateMenuCustomizationFromMultibox(SharedThis(this));
+					}
+				}
 			}
 		}
+		else
+		{
+			if (ExistingIndex != INDEX_NONE)
+			{
+				Blocks.RemoveAt(ExistingIndex);
+				if (ExistingIndex < Index)
+				{
+					--Index;
+				}
+			}
 
-		Blocks.Insert( InBlock, Index );
+			Blocks.Insert(InBlock, Index);
 
-		CustomizationData->BlockAdded( InBlock, Index, Blocks );
+			if (UToolMenuBase* ToolMenu = GetToolMenu())
+			{
+				ToolMenu->UpdateMenuCustomizationFromMultibox(SharedThis(this));
+			}
+		}
 	}
 }
 
@@ -259,8 +334,6 @@ void FMultiBox::InsertCustomMultiBlock( TSharedRef<const FMultiBlock> InBlock, i
  */
 TSharedRef< SMultiBoxWidget > FMultiBox::MakeWidget( bool bSearchable, FOnMakeMultiBoxBuilderOverride* InMakeMultiBoxBuilderOverride /* = nullptr */ )
 {	
-	ApplyCustomizedBlocks();
-
 	TSharedRef< SMultiBoxWidget > NewMultiBoxWidget =
 		SNew( SMultiBoxWidget );
 
@@ -293,69 +366,17 @@ TSharedRef< SMultiBoxWidget > FMultiBox::MakeWidget( bool bSearchable, FOnMakeMu
 
 bool FMultiBox::IsCustomizable() const
 {
-	bool bIsCustomizable = false;
-	if( CustomizationData->GetCustomizationName() != NAME_None )
+	if (UToolMenuBase* ToolMenu = GetToolMenu())
 	{
-		GConfig->GetBool(TEXT("/Script/UnrealEd.EditorExperimentalSettings"), TEXT("bToolbarCustomization"), bIsCustomizable, GEditorPerProjectIni);
+		return ToolMenu->IsEditing();
 	}
 
-	return bIsCustomizable;
-}
-
-void FMultiBox::ApplyCustomizedBlocks()
-{
-	if( IsCustomizable() )
-	{
-		CustomizationData->LoadCustomizedBlocks();
-
-		// Build a map of commands to existing blocks,  we'll try to use existing blocks before creating new ones
-		TMap< TSharedPtr<const FUICommandInfo>, TSharedPtr<const FMultiBlock> > CommandToBlockMap;
-
-		for( int32 BlockIndex = 0; BlockIndex < Blocks.Num(); ++BlockIndex )
-		{
-			TSharedPtr<const FMultiBlock> Block = Blocks[BlockIndex];
-			if( Block->GetAction().IsValid() )
-			{
-				CommandToBlockMap.Add( Block->GetAction(), Block );
-			}
-		}
-
-		// Rebuild the users customized box by executing the transactions the user made to get the 
-		// box to its customized state
-		for( uint32 TransIndex = 0; TransIndex < CustomizationData->GetNumTransactions(); ++TransIndex )
-		{
-			const FCustomBlockTransaction& Transaction = CustomizationData->GetTransaction(TransIndex);
-		
-			// Try and find the block in the default map;
-			TSharedPtr<const FMultiBlock> Block = CommandToBlockMap.FindRef( Transaction.Command.Pin() );
-
-			if( Transaction.TransactionType == FCustomBlockTransaction::Add )
-			{
-			
-				if( !Block.IsValid() )
-				{
-					Block = MakeMultiBlockFromCommand( Transaction.Command.Pin(), false );
-				}
-
-				if( Block.IsValid() )
-				{
-					Blocks.Insert( Block.ToSharedRef(), FMath::Clamp( Transaction.BlockIndex, 0, Blocks.Num() ) );
-				}
-			}
-			else
-			{
-				if( Block.IsValid() )
-				{
-					Blocks.Remove( Block.ToSharedRef() );
-				}
-			}
-		}
-	}
+	return false;
 }
 
 FName FMultiBox::GetCustomizationName() const
 {
-	return CustomizationData->GetCustomizationName(); 
+	return NAME_None;
 }
 
 TSharedPtr<FMultiBlock> FMultiBox::MakeMultiBlockFromCommand( TSharedPtr<const FUICommandInfo> CommandInfo, bool bCommandMustBeBound ) const
@@ -405,20 +426,74 @@ TSharedPtr<FMultiBlock> FMultiBox::MakeMultiBlockFromCommand( TSharedPtr<const F
 
 }
 
-TSharedPtr<const FMultiBlock> FMultiBox::FindBlockFromCommand( TSharedPtr<const FUICommandInfo> Command ) const
+TSharedPtr<const FMultiBlock> FMultiBox::FindBlockFromNameAndType(const FName InName, const EMultiBlockType InType) const
 {
-	for (TArray< TSharedRef< const FMultiBlock > >::TConstIterator It(Blocks); It; ++It)
+	for (const auto& Block : Blocks)
 	{
-		const TSharedRef< const FMultiBlock >& Block = *It;
-		if( Block->GetAction() == Command )
+		if (Block->GetExtensionHook() == InName && Block->GetType() == InType)
 		{
 			return Block;
 		}
 	}
 
-	return NULL;
+	return nullptr;
+}
+
+int32 FMultiBox::GetSectionEditBounds(const int32 Index, int32& OutSectionEndIndex) const
+{
+	// Only used by edit mode, identifies sections by heading blocks
+	if (!IsInEditMode())
+	{
+		return INDEX_NONE;
+	}
+
+	int32 SectionBeginIndex = INDEX_NONE;
+	for (int32 BlockIdx = Index; BlockIdx >= 0; --BlockIdx)
+	{
+		if (Blocks[BlockIdx]->GetType() == EMultiBlockType::Heading)
+		{
+			if (BlockIdx > 0 && Blocks[BlockIdx - 1]->IsSeparator() && Blocks[BlockIdx]->GetExtensionHook() == Blocks[BlockIdx - 1]->GetExtensionHook())
+			{
+				SectionBeginIndex = BlockIdx - 1;
+			}
+			else
+			{
+				SectionBeginIndex = BlockIdx;
+			}
+			break;
+		}
+	}
+
+	OutSectionEndIndex = Blocks.Num();
+	for (int32 BlockIdx = Index + 1; BlockIdx < Blocks.Num(); ++BlockIdx)
+	{
+		if (Blocks[BlockIdx]->GetType() == EMultiBlockType::Heading)
+		{
+			if (BlockIdx > 0 && Blocks[BlockIdx - 1]->IsSeparator() && Blocks[BlockIdx]->GetExtensionHook() == Blocks[BlockIdx - 1]->GetExtensionHook())
+			{
+				OutSectionEndIndex = BlockIdx - 1;
+			}
+			else
+			{
+				OutSectionEndIndex = BlockIdx;
+			}
+			break;
+		}
+	}
+
+	return SectionBeginIndex;
 }
 
+UToolMenuBase* FMultiBox::GetToolMenu() const
+{
+	return WeakToolMenu.Get();
+}
+
+bool FMultiBox::IsInEditMode() const
+{
+	UToolMenuBase* ToolMenu = GetToolMenu();
+	return ToolMenu && ToolMenu->IsEditing();
+}
 
 void SMultiBoxWidget::Construct( const FArguments& InArgs )
 {
@@ -463,6 +538,37 @@ bool SMultiBoxWidget::IsBlockBeingDragged( TSharedPtr<const FMultiBlock> Block )
 	return false;
 }
 
+EVisibility SMultiBoxWidget::GetCustomizationBorderDragVisibility(const FName InBlockName, const EMultiBlockType InBlockType, bool& bOutInsertAfter) const
+{
+	bOutInsertAfter = false;
+
+	if (DragPreview.PreviewBlock.IsValid())
+	{
+		const TArray< TSharedRef< const FMultiBlock > >& Blocks = MultiBox->GetBlocks();
+		if (Blocks.IsValidIndex(DragPreview.InsertIndex))
+		{
+			if (InBlockName != NAME_None)
+			{
+				const TSharedRef< const FMultiBlock >& DropDestination = Blocks[DragPreview.InsertIndex];
+				if (DropDestination->GetExtensionHook() == InBlockName && DropDestination->GetType() == InBlockType)
+				{
+					return EVisibility::Visible;
+				}
+			}
+		}
+		else if (Blocks.Num() == DragPreview.InsertIndex)
+		{
+			if (Blocks.Num() > 0 && Blocks.Last()->GetExtensionHook() == InBlockName && Blocks.Last()->GetType() == InBlockType)
+			{
+				bOutInsertAfter = true;
+				return EVisibility::Visible;
+			}
+		}
+	}
+
+	return EVisibility::Collapsed;
+}
+
 void SMultiBoxWidget::AddBlockWidget( const FMultiBlock& Block, TSharedPtr<SHorizontalBox> HorizontalBox, TSharedPtr<SVerticalBox> VerticalBox, EMultiBlockLocation::Type InLocation, bool bSectionContainsIcons )
 {
 	check( MultiBox.IsValid() );
@@ -476,31 +582,18 @@ void SMultiBoxWidget::AddBlockWidget( const FMultiBlock& Block, TSharedPtr<SHori
 
 	const ISlateStyle* const StyleSet = MultiBox->GetStyleSet();
 
-	TSharedRef<SWidget> FinalWidget =
-	SNew( SOverlay )
-	+ SOverlay::Slot()
-	[
-		BlockWidget
-	]
-	+ SOverlay::Slot()
-	[
-		// This overlay prevents users from clicking on the actual block when in edit mode and also allows new blocks
-		// to be dropped on disabled blocks
-		SNew( SMultiBlockDragHandle, SharedThis( this ), Block.AsShared(), MultiBox->GetCustomizationName() )
-		.Visibility( this, &SMultiBoxWidget::GetCustomizationVisibility, BlockWeakPtr, BlockWidgetWeakPtr )
-	]
-	+ SOverlay::Slot()
-	.HAlign( HAlign_Right )
-	.VAlign( VAlign_Top )
-	.Padding( FMargin(0.0f, 2.0f, 1.0f, 0.0f ) )
-	[
-		// The delete button for removing blocks is only visible when in edit mode
-		SNew( SButton )
-		.Visibility( this, &SMultiBoxWidget::GetCustomizationVisibility, BlockWeakPtr, BlockWidgetWeakPtr )
-		.ContentPadding(0)
-		.OnClicked( this, &SMultiBoxWidget::OnDeleteBlockClicked, BlockWeakPtr )
-		.ButtonStyle( StyleSet, "MultiBox.DeleteButton" )
-	];
+	TSharedPtr<SWidget> FinalWidget;
+
+	const EMultiBlockType BlockType = Block.GetType();
+
+	if (MultiBox->ModifyBlockWidgetAfterMake.IsBound())
+	{
+		FinalWidget = MultiBox->ModifyBlockWidgetAfterMake.Execute(SharedThis(this), Block, BlockWidget);
+	}
+	else
+	{
+		FinalWidget = BlockWidget;
+	}
 
 	switch (MultiBox->GetType())
 	{
@@ -524,7 +617,7 @@ void SMultiBoxWidget::AddBlockWidget( const FMultiBlock& Block, TSharedPtr<SHori
 				]
 				+SVerticalBox::Slot()
 				[
-					FinalWidget
+					FinalWidget.ToSharedRef()
 				]
 			];
 		}
@@ -547,14 +640,14 @@ void SMultiBoxWidget::AddBlockWidget( const FMultiBlock& Block, TSharedPtr<SHori
 					]
 					+SVerticalBox::Slot()
 					[
-						FinalWidget
+						FinalWidget.ToSharedRef()
 					]
 				];
 		}
 		break;
 	case EMultiBoxType::ButtonRow:
 		{
-			TileViewWidgets.Add( FinalWidget );
+			TileViewWidgets.Add( FinalWidget.ToSharedRef() );
 		}
 		break;
 	case EMultiBoxType::Menu:
@@ -575,7 +668,7 @@ void SMultiBoxWidget::AddBlockWidget( const FMultiBlock& Block, TSharedPtr<SHori
 				]
 				+SHorizontalBox::Slot()
 				[
-					FinalWidget
+					FinalWidget.ToSharedRef()
 				]
 			];
 		}
@@ -832,34 +925,33 @@ TSharedRef<SWidget> SMultiBoxWidget::OnWrapButtonClicked()
 
 void SMultiBoxWidget::UpdateDropAreaPreviewBlock( TSharedRef<const FMultiBlock> MultiBlock, TSharedPtr<FUICommandDragDropOp> DragDropContent, const FGeometry& DragAreaGeometry, const FVector2D& DragPos )
 {
-	TSharedPtr<const FUICommandInfo> UICommand = DragDropContent->UICommand;
+	const FName BlockName = DragDropContent->ItemName;
+	const EMultiBlockType BlockType = DragDropContent->BlockType;
 	FName OriginMultiBox = DragDropContent->OriginMultiBox;
 
 	FVector2D LocalDragPos = DragAreaGeometry.AbsoluteToLocal( DragPos );
 
 	FVector2D DrawSize = DragAreaGeometry.GetDrawSize();
 
+	bool bIsDraggingSection = DragDropContent->bIsDraggingSection;
+
 	bool bAddedNewBlock = false;
 	bool bValidCommand = true;
-	if( DragPreview.UICommand != UICommand )
-	{	
-		TSharedPtr<const FMultiBlock> ExistingBlock = MultiBox->FindBlockFromCommand( UICommand );
-
-		// Check that the command does not already exist and that we can create it or that we are dragging an exisiting block in this box
+	if (!DragPreview.IsSameBlockAs(BlockName, BlockType))
+	{
+		TSharedPtr<const FMultiBlock> ExistingBlock = MultiBox->FindBlockFromNameAndType(BlockName, BlockType);
+		// Check that the command does not already exist and that we can create it or that we are dragging an existing block in this box
 		if( !ExistingBlock.IsValid() || ( ExistingBlock.IsValid() && OriginMultiBox == MultiBox->GetCustomizationName() ) )
 		{
 
 			TSharedPtr<const FMultiBlock> NewBlock = ExistingBlock;
 
-			if( !ExistingBlock.IsValid() )
-			{
-				NewBlock = MultiBox->MakeMultiBlockFromCommand( UICommand, true );
-			}
 
 			if( NewBlock.IsValid() )
 			{
 				DragPreview.Reset();
-				DragPreview.UICommand = UICommand;
+				DragPreview.BlockName = BlockName;
+				DragPreview.BlockType = BlockType;
 				DragPreview.PreviewBlock = 
 					MakeShareable(
 						new FDropPreviewBlock( 
@@ -919,29 +1011,50 @@ void SMultiBoxWidget::UpdateDropAreaPreviewBlock( TSharedRef<const FMultiBlock>
 		if( DragPreview.PreviewBlock.IsValid() )
 		{
 			const TArray< TSharedRef< const FMultiBlock > >& Blocks = MultiBox->GetBlocks();
-			for (int32 BlockIdx = 0; BlockIdx < Blocks.Num(); ++BlockIdx)
+			int32 HoverIndex = Blocks.IndexOfByKey(MultiBlock);
+			int32 HoverSectionEndIndex = INDEX_NONE;
+			int32 HoverSectionBeginIndex = MultiBox->GetSectionEditBounds(HoverIndex, HoverSectionEndIndex);
+
+			if (bIsDraggingSection)
 			{
-				TSharedRef<const FMultiBlock> Block = Blocks[BlockIdx];
-				if( Block == MultiBlock )
+				// Hovering over final block means insert at end of list
+				if ((HoverIndex == Blocks.Num() - 1) && Blocks.Num() > 0)
 				{
-					if( bInsertBefore)
+					DragPreview.InsertIndex = Blocks.Num();
+				}
+				else if (Blocks.IsValidIndex(HoverSectionBeginIndex))
+				{
+					DragPreview.InsertIndex = HoverSectionBeginIndex;
+				}
+			}
+			else if (HoverIndex != INDEX_NONE)
+			{
+				if (Blocks[HoverIndex]->IsPartOfHeading())
+				{
+					if (Blocks[HoverIndex]->IsSeparator())
 					{
-						DragPreview.InsertIndex = BlockIdx;
+						// Move insert index above separator of heading
+						DragPreview.InsertIndex = HoverIndex;
 					}
 					else
 					{
-						DragPreview.InsertIndex = FMath::Min(Blocks.Num()-1, BlockIdx+1);
+						// Move insert index after heading
+						DragPreview.InsertIndex = HoverIndex + 1;
+					}
+				}
+				else
+				{
+					if (bInsertBefore)
+					{
+						DragPreview.InsertIndex = HoverIndex;
+					}
+					else
+					{
+						DragPreview.InsertIndex = HoverIndex + 1;
 					}
-
-					break;
 				}
 			}
 		}
-
-		if( CurrentIndex != DragPreview.InsertIndex && DragPreview.InsertIndex != INDEX_NONE )
-		{
-			BuildMultiBoxWidget();
-		}
 	}
 
 }
@@ -959,17 +1072,6 @@ EVisibility SMultiBoxWidget::GetCustomizationVisibility( TWeakPtr<const FMultiBl
 	}
 }
 
-FReply SMultiBoxWidget::OnDeleteBlockClicked( TWeakPtr<const FMultiBlock> BlockWeakPtr )
-{
-	if( BlockWeakPtr.IsValid() )
-	{
-		MultiBox->RemoveCustomMultiBlock( BlockWeakPtr.Pin().ToSharedRef() );
-		BuildMultiBoxWidget();
-	}
-
-	return FReply::Handled();
-}
-
 void SMultiBoxWidget::OnCustomCommandDragEnter( TSharedRef<const FMultiBlock> MultiBlock, const FGeometry& MyGeometry, const FDragDropEvent& DragDropEvent )
 {
 	if( MultiBlock != DragPreview.PreviewBlock && MultiBox->IsInEditMode() )
@@ -997,14 +1099,18 @@ void SMultiBoxWidget::OnCustomCommandDropped()
 	{	
 
 		// Check that the command does not already exist and that we can create it or that we are dragging an exisiting block in this box
-		TSharedPtr<const FMultiBlock> Block = MultiBox->FindBlockFromCommand( DragPreview.UICommand );
-		if( !Block.IsValid() )
+		TSharedPtr<const FMultiBlock> Block = MultiBox->FindBlockFromNameAndType(DragPreview.BlockName, DragPreview.BlockType);
+		if(Block.IsValid())
 		{
-			Block = MultiBox->MakeMultiBlockFromCommand( DragPreview.UICommand, true );
-		}
+			if (Block->IsSeparator() && Block->IsPartOfHeading())
+			{
+				TSharedPtr<const FMultiBlock> HeadingBlock = MultiBox->FindBlockFromNameAndType(DragPreview.BlockName, EMultiBlockType::Heading);
+				if (HeadingBlock.IsValid())
+				{
+					Block = HeadingBlock;
+				}
+			}
 
-		if( Block.IsValid() )
-		{
 			MultiBox->InsertCustomMultiBlock( Block.ToSharedRef(), DragPreview.InsertIndex );
 		}
 
diff --git a/Engine/Source/Runtime/Slate/Private/Framework/MultiBox/MultiBoxBuilder.cpp b/Engine/Source/Runtime/Slate/Private/Framework/MultiBox/MultiBoxBuilder.cpp
--- a/Engine/Source/Runtime/Slate/Private/Framework/MultiBox/MultiBoxBuilder.cpp
+++ b/Engine/Source/Runtime/Slate/Private/Framework/MultiBox/MultiBoxBuilder.cpp
@@ -16,6 +16,8 @@
 #include "Framework/MultiBox/SButtonRowBlock.h"
 #include "Framework/MultiBox/SWidgetBlock.h"
 #include "Framework/MultiBox/SGroupMarkerBlock.h"
+#include "Framework/MultiBox/ToolMenuBase.h"
+
 
 FMultiBoxBuilder::FMultiBoxBuilder( const EMultiBoxType InType, FMultiBoxCustomization InCustomization, const bool bInShouldCloseWindowAfterMenuSelection, const TSharedPtr< const FUICommandList >& InCommandList, TSharedPtr<FExtender> InExtender, FName InTutorialHighlightName, FName InMenuName )
 	: MultiBox( FMultiBox::Create( InType, InMenuName != NAME_None ? FMultiBoxCustomization::AllowCustomization(InMenuName) : InCustomization, bInShouldCloseWindowAfterMenuSelection ) )
@@ -214,7 +216,7 @@ void FMenuBuilder::AddMenuSeparator(FName InExtensionHook)
 	// Never add a menu separate as the first item, even if we were asked to
 	if( MultiBox->GetBlocks().Num() > 0 || FMultiBoxSettings::DisplayMultiboxHooks.Get() )
 	{
-		TSharedRef< FMenuSeparatorBlock > NewMenuSeparatorBlock( new FMenuSeparatorBlock(InExtensionHook) );
+		TSharedRef< FMenuSeparatorBlock > NewMenuSeparatorBlock( new FMenuSeparatorBlock(InExtensionHook, /* bInIsPartOfHeading=*/ false) );
 		MultiBox->AddMultiBlock( NewMenuSeparatorBlock );
 	}
 
@@ -317,7 +319,7 @@ void FMenuBuilder::ApplySectionBeginning()
 		// Do not count search block, which starts as invisible
 		if( MultiBox->GetBlocks().Num() > 1 || FMultiBoxSettings::DisplayMultiboxHooks.Get() )
 		{
-			MultiBox->AddMultiBlock( MakeShareable( new FMenuSeparatorBlock(CurrentSectionExtensionHook) ) );
+			MultiBox->AddMultiBlock( MakeShareable( new FMenuSeparatorBlock(CurrentSectionExtensionHook, /* bInIsPartOfHeading=*/ true) ) );
 		}
 		if (!CurrentSectionHeadingText.IsEmpty())
 		{
diff --git a/Engine/Source/Runtime/Slate/Private/Framework/MultiBox/MultiBoxCustomization.cpp b/Engine/Source/Runtime/Slate/Private/Framework/MultiBox/MultiBoxCustomization.cpp
--- a/Engine/Source/Runtime/Slate/Private/Framework/MultiBox/MultiBoxCustomization.cpp
+++ b/Engine/Source/Runtime/Slate/Private/Framework/MultiBox/MultiBoxCustomization.cpp
@@ -1,98 +1,11 @@
 // Copyright 1998-2019 Epic Games, Inc. All Rights Reserved.
 
 #include "Framework/MultiBox/MultiBoxCustomization.h"
-#include "Framework/MultiBox/ToolMenuBase.h"
-#include "Misc/ConfigCacheIni.h"
-#include "Serialization/JsonReader.h"
-#include "Serialization/JsonSerializer.h"
 #include "Framework/Commands/InputBindingManager.h"
 #include "Widgets/Layout/SBorder.h"
-#include "Misc/RemoteConfigIni.h"
 #include "Framework/Commands/UICommandDragDropOp.h"
 #include "Widgets/Colors/SColorBlock.h"
 
-FCustomizedToolMenuSection* FCustomizedToolMenu::FindSection(const FName InSectionName)
-{
-	for (FCustomizedToolMenuSection& Section : Sections)
-	{
-		if (Section.Name == InSectionName)
-		{
-			return &Section;
-		}
-	}
-
-	return nullptr;
-}
-
-const FCustomizedToolMenuSection* FCustomizedToolMenu::FindSection(const FName InSectionName) const
-{
-	for (const FCustomizedToolMenuSection& Section : Sections)
-	{
-		if (Section.Name == InSectionName)
-		{
-			return &Section;
-		}
-	}
-
-	return nullptr;
-}
-
-FCustomizedToolMenuEntry* FCustomizedToolMenu::FindEntry(const FName InEntryName, FName* OutSectionName)
-{
-	for (FCustomizedToolMenuSection& Section : Sections)
-	{
-		for (FCustomizedToolMenuEntry& Entry : Section.Entries)
-		{
-			if (Entry.Name == InEntryName)
-			{
-				if (OutSectionName)
-				{
-					*OutSectionName = Section.Name;
-				}
-				return &Entry;
-			}
-		}
-	}
-
-	return nullptr;
-}
-
-const FCustomizedToolMenuEntry* FCustomizedToolMenu::FindEntry(const FName InEntryName, FName* OutSectionName) const
-{
-	for (const FCustomizedToolMenuSection& Section : Sections)
-	{
-		for (const FCustomizedToolMenuEntry& Entry : Section.Entries)
-		{
-			if (Entry.Name == InEntryName)
-			{
-				if (OutSectionName)
-				{
-					*OutSectionName = Section.Name;
-				}
-				return &Entry;
-			}
-		}
-	}
-
-	return nullptr;
-}
-
-bool FCustomizedToolMenu::IsSectionHidden(const FName InSectionName) const
-{
-	return HiddenSections.Contains(InSectionName);
-}
-
-bool FCustomizedToolMenu::IsEntryHidden(const FName InEntryName) const
-{
-	return HiddenEntries.Contains(InEntryName);
-}
-
-void FCustomizedToolMenu::UpdateFromMultiBox(const TSharedRef<const FMultiBox>& InMultiBox)
-{
-
-}
-
-
 void SCustomToolbarPreviewWidget::Construct( const FArguments& InArgs )
 {
 	Content = InArgs._Content.Widget;
@@ -114,69 +27,9 @@ void SCustomToolbarPreviewWidget::BuildMultiBlockWidget(const ISlateStyle* Style
 
 	// Add this widget to the search list of the multibox and hide it
 	if (MultiBlock->GetSearchable())
-		OwnerMultiBoxWidget.Pin()->AddSearchElement(this->AsWidget(), FText::GetEmpty());
-}
-
-void SMultiBlockDragHandle::Construct( const FArguments& InArgs, TSharedRef<SMultiBoxWidget> InBaseWidget, TSharedRef<const FMultiBlock> InBlock, FName CustomizationName )
-{
-	BaseWidget = InBaseWidget;
-	Block = InBlock;
-	MultiBoxCustomizationName = CustomizationName;
-}
-
-FReply SMultiBlockDragHandle::OnMouseButtonDown( const FGeometry& MyGeometry, const FPointerEvent& MouseEvent )
-{
-	if( MouseEvent.GetEffectingButton() == EKeys::LeftMouseButton && Block->GetAction().IsValid() )
-	{
-		return FReply::Handled().DetectDrag( SharedThis(this), MouseEvent.GetEffectingButton() );
-	}
-
-	return FReply::Unhandled();
-}
-
-void SMultiBlockDragHandle::OnDragEnter( const FGeometry& MyGeometry, const FDragDropEvent& DragDropEvent )
-{
-	if ( DragDropEvent.GetOperationAs<FUICommandDragDropOp>().IsValid() )
-	{
-		BaseWidget.Pin()->OnCustomCommandDragEnter( Block.ToSharedRef(), MyGeometry, DragDropEvent );
-	}
-}
-
-FReply SMultiBlockDragHandle::OnDragOver( const FGeometry& MyGeometry, const FDragDropEvent& DragDropEvent )
-{
-	if ( DragDropEvent.GetOperationAs<FUICommandDragDropOp>().IsValid() )
-	{
-		BaseWidget.Pin()->OnCustomCommandDragged( Block.ToSharedRef(), MyGeometry, DragDropEvent );
-		return FReply::Handled();
-	}
-
-	return FReply::Unhandled();
-}
-
-FReply SMultiBlockDragHandle::OnDrop( const FGeometry& MyGeometry, const FDragDropEvent& DragDropEvent )
-{
-	if ( DragDropEvent.GetOperationAs<FUICommandDragDropOp>().IsValid() )
 	{
-		BaseWidget.Pin()->OnCustomCommandDropped();
-		return FReply::Handled();
+		OwnerMultiBoxWidget.Pin()->AddSearchElement(this->AsWidget(), FText::GetEmpty());
 	}
-
-	return FReply::Unhandled();
-}
-
-FReply SMultiBlockDragHandle::OnDragDetected( const FGeometry& MyGeometry, const FPointerEvent& MouseEvent )
-{
-	TSharedRef<FUICommandDragDropOp> NewOp = FUICommandDragDropOp::New(
-			Block->GetAction().ToSharedRef(), 
-			MultiBoxCustomizationName, 
-			Block->MakeWidget( BaseWidget.Pin().ToSharedRef(), EMultiBlockLocation::None, Block->HasIcon() )->AsWidget(),
-			MyGeometry.AbsolutePosition-MouseEvent.GetScreenSpacePosition()
-		);
-			
-	NewOp->SetOnDropNotification( FSimpleDelegate::CreateSP( BaseWidget.Pin().ToSharedRef(), &SMultiBoxWidget::OnDropExternal ) );
-
-	TSharedRef<FDragDropOperation> DragDropOp = NewOp;
-	return FReply::Handled().BeginDragDrop( DragDropOp );
 }
 
 TSharedRef< class IMultiBlockBaseWidget > FDropPreviewBlock::ConstructWidget() const
@@ -190,294 +43,3 @@ TSharedRef< class IMultiBlockBaseWidget > FDropPreviewBlock::ConstructWidget() c
 		];
 
 }
-
-void FMultiBoxCustomizationData::LoadCustomizedBlocks()
-{
-	Transactions.Empty();
-
-	FString Content;
-	
-	GConfig->GetString(*GetConfigSectionName(), *CustomizationName.ToString(), Content, GEditorPerProjectIni);
-
-	TSharedPtr<FJsonObject> SavedData;
-	TSharedRef< TJsonReader<> > Reader = TJsonReaderFactory<>::Create( FRemoteConfig::ReplaceIniSpecialCharWithChar(Content).ReplaceEscapedCharWithChar() );
-	bool bResult = FJsonSerializer::Deserialize( Reader, SavedData );
-
-	if( bResult && SavedData.IsValid() )
-	{
-		TArray<TSharedPtr<FJsonValue> > CustomBlockData = SavedData->GetArrayField( TEXT("CustomBlocks") );
-		for (TArray<TSharedPtr<FJsonValue> >::TConstIterator It(CustomBlockData); It; ++It)
-		{
-			TSharedPtr<FJsonObject> SavedCommand = (*It)->AsObject();
-
-			if( SavedCommand.IsValid() )
-			{
-				FString CommandName = SavedCommand->GetStringField( TEXT("CommandName") );
-				FString Context = SavedCommand->GetStringField( TEXT("Context") );
-				int32 Index = FMath::TruncToInt( SavedCommand->GetNumberField( TEXT("Index") ) );
-				FCustomBlockTransaction::ETransactionType TransType = (FCustomBlockTransaction::ETransactionType)FMath::TruncToInt( SavedCommand->GetNumberField("TransactionType") );
-
-				if( !CommandName.IsEmpty() && !Context.IsEmpty() && Index >= 0 )
-				{
-					const TSharedPtr< const FUICommandInfo> FoundCommand = FInputBindingManager::Get().FindCommandInContext( FName( *Context ), FName( *CommandName ) );
-					if( FoundCommand.IsValid() )
-					{
-						FCustomBlockTransaction NewTransaction;
-						if( TransType == FCustomBlockTransaction::Add )
-						{
-							NewTransaction = FCustomBlockTransaction::CreateAdd( FoundCommand.ToSharedRef(), Index );
-						}
-						else
-						{
-							NewTransaction = FCustomBlockTransaction::CreateRemove( FoundCommand.ToSharedRef(), Index );
-						}
-
-						Transactions.Add( NewTransaction );
-					}
-				}
-			}
-		}
-	}
-}
-
-FString FMultiBoxCustomizationData::GetConfigSectionName() const
-{
-	return TEXT("CustomMultiBoxes1_0");
-}
-
-void FMultiBoxCustomizationData::RemoveTransactionAt(int32 RemoveIndex)
-{
-	// Copy off the transaction we are going to remove.
-	FCustomBlockTransaction TransToRemove = Transactions[RemoveIndex];
-	const int32 TransToRemoveIndex = TransToRemove.BlockIndex;
-
-	// Remove the transaction
-	Transactions.RemoveAt( RemoveIndex );
-
-	// Iterate over all transactions after this one to fix up indices.
-	for( int32 StartIndex = RemoveIndex; StartIndex < Transactions.Num(); ++StartIndex )
-	{
-		int32& CheckIndex = Transactions[StartIndex].BlockIndex;
-		
-		// If the index is greater than or equal to the index we are removing, this transaction is affected by this removal.
-		if( CheckIndex >= TransToRemoveIndex )
-		{
-			// If the removed transaction was an add, subtract 1 from the index of this transaction. Otherwise, add one.
-			if( TransToRemove.TransactionType == FCustomBlockTransaction::Add )
-			{
-				--CheckIndex;
-			}
-			else
-			{
-				++CheckIndex;
-			}
-		}
-	}
-}
-
-bool FMultiBoxCustomizationData::RemoveDuplicateTransaction()
-{
-	if( Transactions.Num() > 0 )
-	{
-		FCustomBlockTransaction LastTrans = Transactions.Last();
-
-		int32 CheckIndex = LastTrans.BlockIndex;
-		
-		// Skip the last item
-		for( int32 TransIndex = Transactions.Num()-2; TransIndex >= 0; --TransIndex )
-		{
-			FCustomBlockTransaction CurrentTrans = Transactions[TransIndex];
-
-			//ensure( !(CurrentTrans.TransactionType == LastTrans.TransactionType && CurrentTrans.Command == LastTrans.Command ) );
-
-			if( CurrentTrans.Command == LastTrans.Command &&
-				CurrentTrans.BlockIndex == CheckIndex &&
-				CurrentTrans.TransactionType != LastTrans.TransactionType )
-			{
-				RemoveTransactionAt( Transactions.Num()-1 );
-				RemoveTransactionAt( TransIndex );
-
-				return true;
-			}
-
-			if( CheckIndex >= CurrentTrans.BlockIndex )
-			{
-				if( CurrentTrans.TransactionType == FCustomBlockTransaction::Add )
-				{
-					--CheckIndex;
-				}
-				else
-				{
-					++CheckIndex;
-				}
-			}
-		}
-	}
-
-	return false;
-}
-
-bool FMultiBoxCustomizationData::RemoveUnnecessaryTransactions(const TArray< TSharedRef< const FMultiBlock > >& AllBlocks)
-{
-	// A local struct describing the state of the menu
-	struct FCustomizationState
-	{
-		FCustomizationState(const TArray< TSharedRef< const FMultiBlock > >& InAllBlocks)
-		{
-			for ( int32 BlockIdx = 0; BlockIdx < InAllBlocks.Num(); ++BlockIdx )
-			{
-				StateData.Add(InAllBlocks[BlockIdx]->GetAction());
-			}
-		}
-
-		// Applies the inverse operation of the given transaction to mutate the state
-		void ApplyInverseOfTransaction(const FCustomBlockTransaction& Trans)
-		{
-			if ( Trans.TransactionType == FCustomBlockTransaction::Add )
-			{
-				// Remove on an add operation
-				if ( ensure(StateData.IsValidIndex(Trans.BlockIndex)) )
-				{
-					ensure(StateData[Trans.BlockIndex] == Trans.Command);
-
-					StateData.RemoveAt(Trans.BlockIndex);
-				}
-			}
-			else if ( Trans.TransactionType == FCustomBlockTransaction::Remove )
-			{
-				// Add on a remove transaction
-				StateData.Insert(Trans.Command, Trans.BlockIndex);
-			}
-		}
-
-		// Compares two states. If all elements of the states are the same then they are "equal"
-		bool operator==(const FCustomizationState& Other) const
-		{
-			if ( StateData.Num() != Other.StateData.Num() )
-			{
-				return false;
-			}
-
-			const int32 Length = StateData.Num();
-			for ( int32 StateIdx = 0; StateIdx < Length; ++StateIdx )
-			{
-				if ( StateData[StateIdx] != Other.StateData[StateIdx] )
-				{
-					return false;
-				}
-			}
-
-			return true;
-		}
-
-	private:
-		TArray<TWeakPtr<const FUICommandInfo>> StateData;
-	};
-
-	if( Transactions.Num() > 0 )
-	{
-		// Take the current state and start applying the inverse of each transaction to it to determine the states before the current.
-		// If we found any state that is identical to the current state, all transactions between that state and the current state are unnecessary, so remove them
-
-		// Record the initial (current) state and initialize the Test state.
-		FCustomizationState InitialState(AllBlocks);
-		FCustomizationState TestState = InitialState;
-
-		// Walk backwards through all transactions, applying the inverse of each transaction to the test state
-		for ( int32 TransIdx = Transactions.Num() - 1; TransIdx >= 0; --TransIdx )
-		{
-			TestState.ApplyInverseOfTransaction(Transactions[TransIdx]);
-
-			// If the test state is equal to the current state, all transactions between are unnecessary. Remove them.
-			if ( TestState == InitialState )
-			{
-				for ( int32 RemoveIdx = Transactions.Num() - 1; RemoveIdx >= TransIdx; --RemoveIdx )
-				{
-					RemoveTransactionAt(RemoveIdx);
-				}
-
-				// Return true to indicate we actually removed something.
-				return true;
-			}
-		}
-	}
-
-	// All transactions were necessary to form the final state. Return false.
-	return false;
-}
-
-void FMultiBoxCustomizationData::SaveCustomizedBlocks()
-{
-	FString SaveData;
-	TSharedRef< TJsonWriter<> > Writer = TJsonWriterFactory<>::Create( &SaveData );
-
-	Writer->WriteObjectStart();
-	Writer->WriteArrayStart(TEXT("CustomBlocks"));
-	for (TArray<FCustomBlockTransaction>::TConstIterator It(Transactions); It; ++It)
-	{
-		const FCustomBlockTransaction& Transaction = *It;
-
-		if( Transaction.Command.IsValid() )
-		{
-			const FUICommandInfo& Command = *Transaction.Command.Pin();
-
-			Writer->WriteObjectStart();
-			Writer->WriteValue(TEXT("CommandName"), Command.GetCommandName().ToString() );
-			Writer->WriteValue(TEXT("Context"), Command.GetBindingContext().ToString() );
-			Writer->WriteValue(TEXT("Index"), (float)Transaction.BlockIndex );
-			Writer->WriteValue(TEXT("TransactionType"), (float)Transaction.TransactionType );
-			Writer->WriteObjectEnd();
-		}
-	}
-
-	Writer->WriteArrayEnd();
-	Writer->WriteObjectEnd();
-	Writer->Close();
-
-	GConfig->SetString( *GetConfigSectionName(), *CustomizationName.ToString(), *FRemoteConfig::ReplaceIniCharWithSpecialChar(SaveData).ReplaceCharWithEscapedChar(), GEditorPerProjectIni );
-}
-
-void FMultiBoxCustomizationData::BlockRemoved( TSharedRef< const FMultiBlock> RemovedBlock, int32 Index, const TArray< TSharedRef< const FMultiBlock > >& AllBlocks )
-{
-	FCustomBlockTransaction Remove = FCustomBlockTransaction::CreateRemove( RemovedBlock->GetAction().ToSharedRef(), Index );
-
-	SaveTransaction( Remove, AllBlocks );
-
-}
-
-void FMultiBoxCustomizationData::BlockAdded( TSharedRef< const FMultiBlock > AddedBlock, int32 Index, const TArray< TSharedRef< const FMultiBlock > >& AllBlocks )
-{
-	FCustomBlockTransaction Add = FCustomBlockTransaction::CreateAdd( AddedBlock->GetAction().ToSharedRef(), Index );
-
-	SaveTransaction( Add, AllBlocks );
-}
-
-#define DEBUG_TRANSACTIONS 0
-#if DEBUG_TRANSACTIONS
-static void PrintTransactions( const TArray<FCustomBlockTransaction>& Transactions )
-{
-	FPlatformMisc::LowLevelOutputDebugStringf(TEXT("==========BEGIN TRANSACTIONS=======\n"));
-	for( int32 i = 0; i < Transactions.Num(); ++i )
-	{
-		const FCustomBlockTransaction& Trans = Transactions[i];
-
-		FString TransType = Trans.TransactionType == FCustomBlockTransaction::Add ? TEXT("+") : TEXT("-");
-		FPlatformMisc::LowLevelOutputDebugStringf( TEXT("%s(%s,%d)\n"), *TransType, *Trans.Command.Pin()->GetCommandName().ToString(), Trans.BlockIndex );
-	}
-	FPlatformMisc::LowLevelOutputDebugStringf(TEXT("============END TRANSACTIONS=======\n"));
-}
-#endif
-
-void FMultiBoxCustomizationData::SaveTransaction( const struct FCustomBlockTransaction& Transaction, const TArray< TSharedRef< const FMultiBlock > >& AllBlocks )
-{
-	Transactions.Add( Transaction );
-
-	while( RemoveDuplicateTransaction() );
-	while( RemoveUnnecessaryTransactions(AllBlocks) );
-
-#if DEBUG_TRANSACTIONS
-	PrintTransactions( Transactions );
-#endif
-
-	SaveCustomizedBlocks();
-}
diff --git a/Engine/Source/Runtime/Slate/Private/Framework/MultiBox/MultiBoxCustomization.h b/Engine/Source/Runtime/Slate/Private/Framework/MultiBox/MultiBoxCustomization.h
--- a/Engine/Source/Runtime/Slate/Private/Framework/MultiBox/MultiBoxCustomization.h
+++ b/Engine/Source/Runtime/Slate/Private/Framework/MultiBox/MultiBoxCustomization.h
@@ -23,27 +23,6 @@ private:
 	TSharedPtr<SWidget> Content;
 };
 
-class SMultiBlockDragHandle : public SCompoundWidget
-{
-	SLATE_BEGIN_ARGS( SMultiBlockDragHandle ) {}
-	SLATE_END_ARGS()
-
-	void Construct( const FArguments& InArgs, TSharedRef<SMultiBoxWidget> InBaseWidget, TSharedRef<const FMultiBlock> InBlock, FName CustomizationName );
-
-	FReply OnMouseButtonDown( const FGeometry& MyGeometry, const FPointerEvent& MouseEvent ) override;
-
-	virtual void OnDragEnter( const FGeometry& MyGeometry, const FDragDropEvent& DragDropEvent ) override;
-	virtual FReply OnDragOver( const FGeometry& MyGeometry, const FDragDropEvent& DragDropEvent ) override;
-	virtual FReply OnDrop( const FGeometry& MyGeometry, const FDragDropEvent& DragDropEvent ) override;
-	virtual FReply OnDragDetected( const FGeometry& MyGeometry, const FPointerEvent& MouseEvent ) override;
-
-private:
-	TSharedPtr<const FMultiBlock> Block;
-	TWeakPtr<SMultiBoxWidget> BaseWidget;
-	FName MultiBoxCustomizationName;
-};
-
-
 /**
  * Arbitrary Widget MultiBlock
  */
@@ -72,69 +51,3 @@ private:
 
 };
 
-struct FCustomBlockTransaction
-{
-	enum ETransactionType
-	{
-		Remove,
-		Add
-	};
-
-	TWeakPtr<const FUICommandInfo> Command;
-	int32 BlockIndex;
-	ETransactionType TransactionType;
-
-	bool operator==(const FCustomBlockTransaction& Other ) const
-	{
-		return Command == Other.Command && BlockIndex == Other.BlockIndex && TransactionType == Other.TransactionType;
-	}
-
-	static FCustomBlockTransaction CreateRemove( TSharedRef<const FUICommandInfo> InCommand, int32 Index )
-	{
-		FCustomBlockTransaction Transaction;
-		Transaction.Command = InCommand;
-		Transaction.BlockIndex = Index;
-		Transaction.TransactionType = ETransactionType::Remove;
-
-		return Transaction;
-	}
-
-	static FCustomBlockTransaction CreateAdd( TSharedRef<const FUICommandInfo> InCommand, int32 Index )
-	{
-		FCustomBlockTransaction Transaction;
-		Transaction.Command = InCommand;
-		Transaction.BlockIndex = Index;
-		Transaction.TransactionType = ETransactionType::Add;
-
-		return Transaction;
-	}
-};
-
-class FMultiBoxCustomizationData
-{
-public:
-	FMultiBoxCustomizationData( FName InCustomizationName )
-		: CustomizationName( InCustomizationName )
-	{}
-
-	FName GetCustomizationName() const { return CustomizationName; }
-
-	void LoadCustomizedBlocks();
-
-	void SaveCustomizedBlocks();
-	void BlockRemoved( TSharedRef< const FMultiBlock> RemovedBlock, int32 Index, const TArray< TSharedRef< const FMultiBlock > >& AllBlocks );
-	void BlockAdded( TSharedRef< const FMultiBlock > AddedBlock, int32 Index, const TArray< TSharedRef< const FMultiBlock > >& AllBlocks );
-	void SaveTransaction( const struct FCustomBlockTransaction& Transaction, const TArray< TSharedRef< const FMultiBlock > >& AllBlocks );
-
-	uint32 GetNumTransactions() const { return Transactions.Num(); }
-
-	const FCustomBlockTransaction& GetTransaction( uint32 TransactionIndex ) { return Transactions[TransactionIndex]; }
-private:
-	FString GetConfigSectionName() const;
-	void RemoveTransactionAt(int32 TransIndex);
-	bool RemoveDuplicateTransaction();
-	bool RemoveUnnecessaryTransactions(const TArray< TSharedRef< const FMultiBlock > >& AllBlocks);
-private:
-	TArray<FCustomBlockTransaction> Transactions;
-	FName CustomizationName;
-};
diff --git a/Engine/Source/Runtime/Slate/Private/Framework/MultiBox/SHeadingBlock.cpp b/Engine/Source/Runtime/Slate/Private/Framework/MultiBox/SHeadingBlock.cpp
--- a/Engine/Source/Runtime/Slate/Private/Framework/MultiBox/SHeadingBlock.cpp
+++ b/Engine/Source/Runtime/Slate/Private/Framework/MultiBox/SHeadingBlock.cpp
@@ -9,7 +9,7 @@
  * @param	InHeadingText	Heading text
  */
 FHeadingBlock::FHeadingBlock( const FName& InExtensionHook, const TAttribute< FText >& InHeadingText )
-	: FMultiBlock( NULL, NULL, InExtensionHook, EMultiBlockType::Heading )
+	: FMultiBlock( NULL, NULL, InExtensionHook, EMultiBlockType::Heading, /* bInIsPartOfHeading=*/ true )
 	, HeadingText( InHeadingText )
 {
 }
diff --git a/Engine/Source/Runtime/Slate/Private/Framework/MultiBox/SMenuEntryBlock.cpp b/Engine/Source/Runtime/Slate/Private/Framework/MultiBox/SMenuEntryBlock.cpp
--- a/Engine/Source/Runtime/Slate/Private/Framework/MultiBox/SMenuEntryBlock.cpp
+++ b/Engine/Source/Runtime/Slate/Private/Framework/MultiBox/SMenuEntryBlock.cpp
@@ -1,6 +1,7 @@
 // Copyright 1998-2019 Epic Games, Inc. All Rights Reserved.
 
 #include "Framework/MultiBox/SMenuEntryBlock.h"
+#include "Framework/MultiBox/ToolMenuBase.h"
 #include "Widgets/SBoxPanel.h"
 #include "Layout/WidgetPath.h"
 #include "Framework/Application/SlateApplication.h"
@@ -484,6 +485,11 @@ TSharedRef< SWidget > SMenuEntryBlock::BuildMenuEntryWidget( const FMenuEntryBui
 	{
 		CheckBoxVisibility = EVisibility::Collapsed;
 	}
+	else if (MultiBox->IsInEditMode())
+	{
+		// Hide but keep spacing to avoid confusing user during editing
+		CheckBoxVisibility = EVisibility::Hidden;
+	}
 
 	TAttribute<FSlateColor> CheckBoxForegroundColor = FSlateColor::UseForeground();
 	FName CheckBoxStyle = ISlateStyle::Join( StyleName, ".CheckBox" );
@@ -868,6 +874,16 @@ void SMenuEntryBlock::BuildMultiBlockWidget(const ISlateStyle* StyleSet, const F
 		BuildParams.Label = BuildParams.UICommand.IsValid() ? BuildParams.UICommand->GetLabel() : FText::GetEmpty();
 	}
 
+	const bool bIsEditing = MultiBox->IsInEditMode();
+	if (bIsEditing)
+	{
+		// No dynamic labels while editing
+		if (BuildParams.Label.IsBound())
+		{
+			BuildParams.Label = BuildParams.Label.Get();
+		}
+	}
+
 	// Add this widget to the search list of the multibox
 	// If there is a widget already assigned (created early) ensure that it's STextBlock is set up for searching
 	TSharedPtr< SWidget > ButtonContent = MenuEntryBlock->EntryWidget;
@@ -906,6 +922,15 @@ void SMenuEntryBlock::BuildMultiBlockWidget(const ISlateStyle* StyleSet, const F
 		BuildParams.ToolTip = BuildParams.UICommand.IsValid() ? BuildParams.UICommand->GetDescription() : FText::GetEmpty();
 	}	
 
+	if (bIsEditing)
+	{
+		// No dynamic tooltips while editing
+		if (BuildParams.ToolTip.IsBound())
+		{
+			BuildParams.ToolTip = BuildParams.ToolTip.Get();
+		}
+	}
+
 	if( MultiBox->GetType() == EMultiBoxType::Menu )
 	{
 		if( MenuEntryBlock->bIsSubMenu )
@@ -927,19 +952,32 @@ void SMenuEntryBlock::BuildMultiBlockWidget(const ISlateStyle* StyleSet, const F
 		ChildSlot[ BuildMenuBarWidget( BuildParams ) ];
 	}
 
-	// Insert named widget if desired
-	FName TutorialName = MenuEntryBlock->GetTutorialHighlightName();
-	if(TutorialName != NAME_None)
+	if (!bIsEditing)
 	{
-		TSharedRef<SWidget> ChildWidget = ChildSlot.GetWidget();
-		ChildWidget->AddMetadata<FTagMetaData>(MakeShared<FTagMetaData>(TutorialName));
+		// Insert named widget if desired
+		FName TutorialName = MenuEntryBlock->GetTutorialHighlightName();
+		if (TutorialName != NAME_None)
+		{
+			TSharedRef<SWidget> ChildWidget = ChildSlot.GetWidget();
+			ChildWidget->AddMetadata<FTagMetaData>(MakeShared<FTagMetaData>(TutorialName));
+		}
 	}
 
-	// Bind our widget's enabled state to whether or not our action can execute
-	SetEnabled( TAttribute<bool>( this, &SMenuEntryBlock::IsEnabled ) );
+	if (bIsEditing)
+	{
+		SetEnabled(TAttribute<bool>(this, &SMenuEntryBlock::IsEnabledDuringEditMode));
+	}
+	else
+	{
+		// Bind our widget's enabled state to whether or not our action can execute
+		SetEnabled(TAttribute<bool>(this, &SMenuEntryBlock::IsEnabled));
+	}
 
-	// Bind our widget's visible state to whether or not the action should be visible
-	SetVisibility( TAttribute<EVisibility>(this, &SMenuEntryBlock::GetVisibility) );
+	if (!bIsEditing)
+	{
+		// Bind our widget's visible state to whether or not the action should be visible
+		SetVisibility(TAttribute<EVisibility>(this, &SMenuEntryBlock::GetVisibility));
+	}
 }
 
 void SMenuEntryBlock::RequestSubMenuToggle( bool bOpenMenu, const bool bClobber )
@@ -1084,6 +1122,57 @@ bool SMenuEntryBlock::IsEnabled() const
 	return bEnabled;
 }
 
+bool SMenuEntryBlock::IsEnabledDuringEditMode() const
+{
+	const FName EntryName = MultiBlock->GetExtensionHook();
+	if (EntryName != NAME_None && OwnerMultiBoxWidget.IsValid())
+	{
+		TSharedRef< const FMultiBox > MultiBox = OwnerMultiBoxWidget.Pin()->GetMultiBox();
+		if (UToolMenuBase* ToolMenu = MultiBox->GetToolMenu())
+		{
+			FCustomizedToolMenuHierarchy CustomizationHierarchy = ToolMenu->GetMenuCustomizationHierarchy();
+			if (CustomizationHierarchy.Hierarchy.Num() > 0)
+			{
+				FCustomizedToolMenu FlattenedCustomization = CustomizationHierarchy.GenerateFlattened();
+
+				if (MultiBlock->GetType() == EMultiBlockType::Heading)
+				{
+					if (FlattenedCustomization.GetSectionVisiblity(EntryName) == ECustomizedToolMenuVisibility::Hidden)
+					{
+						return false;
+					}
+				}
+				else
+				{
+					if (FlattenedCustomization.GetEntryVisiblity(EntryName) == ECustomizedToolMenuVisibility::Hidden)
+					{
+						return false;
+					}
+
+					FName SectionName = FlattenedCustomization.GetEntrySectionName(EntryName);
+					if (SectionName == NAME_None)
+					{
+						SectionName = ToolMenu->GetSectionName(EntryName);
+					}
+
+					if (SectionName != NAME_None)
+					{
+						if (FlattenedCustomization.GetSectionVisiblity(SectionName) == ECustomizedToolMenuVisibility::Hidden)
+						{
+							return false;
+						}
+					}
+				}
+			}
+		}
+	}
+	else
+	{
+		return false;
+	}
+
+	return true;
+}
 
 /**
  * Called by Slate when this menu entry check box button is toggled
@@ -1341,7 +1430,18 @@ TSharedRef< SWidget > SMenuEntryBlock::MakeNewMenuWidget() const
 	if (MenuEntryBlock->EntryBuilder.IsBound())
 	{
 		const bool bCloseSelfOnly = false;
-		FMenuBuilder MenuBuilder(MenuEntryBlock->bShouldCloseWindowAfterMenuSelection, MultiBlock->GetActionList(), MenuEntryBlock->Extender, bCloseSelfOnly, StyleSet );
+
+		FName SubMenuCustomizationName;
+		if (MenuEntryBlock->GetExtensionHook() != NAME_None)
+		{
+			FName CustomizationName = MultiBoxWidget->GetMultiBox()->GetCustomizationName();
+			if (CustomizationName != NAME_None)
+			{
+				SubMenuCustomizationName = *(CustomizationName.ToString() + TEXT(".") + MenuEntryBlock->GetExtensionHook().ToString());
+			}
+		}
+
+		FMenuBuilder MenuBuilder(MenuEntryBlock->bShouldCloseWindowAfterMenuSelection, MultiBlock->GetActionList(), MenuEntryBlock->Extender, bCloseSelfOnly, StyleSet, /* searchable */ true, SubMenuCustomizationName );
 		{
 			MenuEntryBlock->EntryBuilder.Execute( MenuBuilder );
 		}
diff --git a/Engine/Source/Runtime/Slate/Private/Framework/MultiBox/SMenuEntryBlock.h b/Engine/Source/Runtime/Slate/Private/Framework/MultiBox/SMenuEntryBlock.h
--- a/Engine/Source/Runtime/Slate/Private/Framework/MultiBox/SMenuEntryBlock.h
+++ b/Engine/Source/Runtime/Slate/Private/Framework/MultiBox/SMenuEntryBlock.h
@@ -238,6 +238,13 @@ protected:
 	 */
 	bool IsEnabled() const;
 
+	/**
+	 * Called by Slate to determine if this menu entry is enabled (during menu editing)
+	 * 
+	 * @return True if the menu entry is enabled, false otherwise
+	 */
+	bool IsEnabledDuringEditMode() const;
+
 	/**
 	 * Called by Slate when this check box button is toggled in a menu entry
 	 */
diff --git a/Engine/Source/Runtime/Slate/Private/Framework/MultiBox/SMenuSeparatorBlock.cpp b/Engine/Source/Runtime/Slate/Private/Framework/MultiBox/SMenuSeparatorBlock.cpp
--- a/Engine/Source/Runtime/Slate/Private/Framework/MultiBox/SMenuSeparatorBlock.cpp
+++ b/Engine/Source/Runtime/Slate/Private/Framework/MultiBox/SMenuSeparatorBlock.cpp
@@ -8,8 +8,8 @@
 /**
  * Constructor
  */
-FMenuSeparatorBlock::FMenuSeparatorBlock(const FName& InExtensionHook)
-	: FMultiBlock( nullptr, nullptr, InExtensionHook, EMultiBlockType::MenuSeparator )
+FMenuSeparatorBlock::FMenuSeparatorBlock(const FName& InExtensionHook, bool bInIsPartOfHeading)
+	: FMultiBlock( nullptr, nullptr, InExtensionHook, EMultiBlockType::MenuSeparator, bInIsPartOfHeading )
 {
 }
 
diff --git a/Engine/Source/Runtime/Slate/Private/Framework/MultiBox/SMenuSeparatorBlock.h b/Engine/Source/Runtime/Slate/Private/Framework/MultiBox/SMenuSeparatorBlock.h
--- a/Engine/Source/Runtime/Slate/Private/Framework/MultiBox/SMenuSeparatorBlock.h
+++ b/Engine/Source/Runtime/Slate/Private/Framework/MultiBox/SMenuSeparatorBlock.h
@@ -18,7 +18,7 @@ public:
 	/**
 	 * Constructor
 	 */
-	FMenuSeparatorBlock(const FName& InExtensionHook);
+	FMenuSeparatorBlock(const FName& InExtensionHook, bool bInIsPartOfHeading);
 
 
 private:
diff --git a/Engine/Source/Runtime/Slate/Private/Framework/MultiBox/ToolMenuCustomization.cpp b/Engine/Source/Runtime/Slate/Private/Framework/MultiBox/ToolMenuCustomization.cpp
new file mode 100644
--- /dev/null
+++ b/Engine/Source/Runtime/Slate/Private/Framework/MultiBox/ToolMenuCustomization.cpp
@@ -0,0 +1,199 @@
+// Copyright 1998-2019 Epic Games, Inc. All Rights Reserved.
+
+#include "Framework/MultiBox/ToolMenuBase.h"
+
+FName FCustomizedToolMenu::GetEntrySectionName(const FName InEntryName) const
+{
+	for (auto& It : EntryOrder)
+	{
+		if (It.Value.Names.Contains(InEntryName))
+		{
+			return It.Key;
+		}
+	}
+
+	return NAME_None;
+}
+
+FCustomizedToolMenuEntry* FCustomizedToolMenu::FindEntry(const FName InEntryName)
+{
+	return Entries.Find(InEntryName);
+	//return Entries.FindByPredicate([=](const FCustomizedToolMenuEntry& Entry) { return Entry.Name == InEntryName; });
+}
+
+const FCustomizedToolMenuEntry* FCustomizedToolMenu::FindEntry(const FName InEntryName) const
+{
+	return Entries.Find(InEntryName);
+	//return Entries.FindByPredicate([=](const FCustomizedToolMenuEntry& Entry) { return Entry.Name == InEntryName; });
+}
+
+FCustomizedToolMenuEntry* FCustomizedToolMenu::AddEntry(const FName InEntryName)
+{
+	return &Entries.FindOrAdd(InEntryName);
+}
+
+ECustomizedToolMenuVisibility FCustomizedToolMenu::GetEntryVisiblity(const FName InEntryName) const
+{
+	if (const FCustomizedToolMenuEntry* Found = FindEntry(InEntryName))
+	{
+		return Found->Visibility;
+	}
+
+	return ECustomizedToolMenuVisibility::None;
+}
+
+bool FCustomizedToolMenu::IsEntryHidden(const FName InEntryName) const
+{
+	return GetEntryVisiblity(InEntryName) == ECustomizedToolMenuVisibility::Hidden;
+}
+
+FCustomizedToolMenuSection* FCustomizedToolMenu::FindSection(const FName InSectionName)
+{
+	return Sections.Find(InSectionName);
+	//return Sections.FindByPredicate([=](const FCustomizedToolMenuSection& Section) { return Section.Name == InSectionName; });
+}
+
+const FCustomizedToolMenuSection* FCustomizedToolMenu::FindSection(const FName InSectionName) const
+{
+	return Sections.Find(InSectionName);
+	//return Sections.FindByPredicate([=](const FCustomizedToolMenuSection& Section) { return Section.Name == InSectionName; });
+}
+
+FCustomizedToolMenuSection* FCustomizedToolMenu::AddSection(const FName InSectionName)
+{
+	return &Sections.FindOrAdd(InSectionName);
+}
+
+ECustomizedToolMenuVisibility FCustomizedToolMenu::GetSectionVisiblity(const FName InSectionName) const
+{
+	if (const FCustomizedToolMenuSection* Found = FindSection(InSectionName))
+	{
+		return Found->Visibility;
+	}
+
+	return ECustomizedToolMenuVisibility::None;
+}
+
+bool FCustomizedToolMenu::IsSectionHidden(const FName InEntryName) const
+{
+	return GetSectionVisiblity(InEntryName) == ECustomizedToolMenuVisibility::Hidden;
+}
+
+FName FCustomizedToolMenuHierarchy::GetEntrySectionName(const FName InEntryName) const
+{
+	for (int32 i = Hierarchy.Num() - 1; i >= 0; --i)
+	{
+		if (Hierarchy[i])
+		{
+			FName SectionName = Hierarchy[i]->GetEntrySectionName(InEntryName);
+			if (SectionName != NAME_None)
+			{
+				return SectionName;
+			}
+		}
+	}
+
+	return NAME_None;
+}
+
+bool FCustomizedToolMenuHierarchy::IsEntryHidden(const FName InEntryName) const
+{
+	for (int32 i = Hierarchy.Num() - 1; i >= 0; --i)
+	{
+		if (Hierarchy[i])
+		{
+			if (const FCustomizedToolMenuEntry* Found = Hierarchy[i]->FindEntry(InEntryName))
+			{
+				if (Found->Visibility == ECustomizedToolMenuVisibility::Hidden)
+				{
+					return true;
+				}
+				else if (Found->Visibility == ECustomizedToolMenuVisibility::Visible)
+				{
+					return false;
+				}
+			}
+		}
+	}
+
+	return false;
+}
+
+bool FCustomizedToolMenuHierarchy::IsSectionHidden(const FName InSectionName) const
+{
+	for (int32 i = Hierarchy.Num() - 1; i >= 0; --i)
+	{
+		if (Hierarchy[i])
+		{
+			if (const FCustomizedToolMenuSection* Found = Hierarchy[i]->FindSection(InSectionName))
+			{
+				if (Found->Visibility == ECustomizedToolMenuVisibility::Hidden)
+				{
+					return true;
+				}
+				else if (Found->Visibility == ECustomizedToolMenuVisibility::Visible)
+				{
+					return false;
+				}
+			}
+		}
+	}
+
+	return false;
+}
+
+FCustomizedToolMenu FCustomizedToolMenuHierarchy::GenerateFlattened() const
+{
+	FCustomizedToolMenu Result;
+	for (const FCustomizedToolMenu* Current : Hierarchy)
+	{
+		if (!Current)
+		{
+			continue;
+		}
+
+		Result.Name = Current->Name;
+
+		if (Current->SectionOrder.Num() > 0)
+		{
+			Result.SectionOrder = Current->SectionOrder;
+		}
+
+		for (const auto& SectionEntryOrderIterator : Current->EntryOrder)
+		{
+			Result.EntryOrder.Add(SectionEntryOrderIterator.Key, SectionEntryOrderIterator.Value);
+		}
+
+		for (const auto& EntryIterator : Current->Entries)
+		{
+			if (FCustomizedToolMenuEntry* ExistingEntry = Result.FindEntry(EntryIterator.Key))
+			{
+				if (EntryIterator.Value.Visibility != ECustomizedToolMenuVisibility::None)
+				{
+					ExistingEntry->Visibility = EntryIterator.Value.Visibility;
+				}
+			}
+			else
+			{
+				Result.Entries.Add(EntryIterator.Key, EntryIterator.Value);
+			}
+		}
+
+		for (const auto& SectionIterator : Current->Sections)
+		{
+			if (FCustomizedToolMenuSection* ExistingSection = Result.FindSection(SectionIterator.Key))
+			{
+				if (SectionIterator.Value.Visibility != ECustomizedToolMenuVisibility::None)
+				{
+					ExistingSection->Visibility = SectionIterator.Value.Visibility;
+				}
+			}
+			else
+			{
+				Result.Sections.Add(SectionIterator.Key, SectionIterator.Value);
+			}
+		}
+	}
+
+	return Result;
+}
diff --git a/Engine/Source/Runtime/Slate/Public/Framework/Commands/UICommandDragDropOp.h b/Engine/Source/Runtime/Slate/Public/Framework/Commands/UICommandDragDropOp.h
--- a/Engine/Source/Runtime/Slate/Public/Framework/Commands/UICommandDragDropOp.h
+++ b/Engine/Source/Runtime/Slate/Public/Framework/Commands/UICommandDragDropOp.h
@@ -5,6 +5,7 @@
 #include "CoreMinimal.h"
 #include "Input/DragAndDrop.h"
 #include "Framework/Commands/UICommandInfo.h"
+#include "Framework/MultiBox/MultiBoxDefs.h"
 
 class SWidget;
 
@@ -18,10 +19,12 @@ public:
 
 	DRAG_DROP_OPERATOR_TYPE(FUICommandDragDropOp, FDragDropOperation)
 
-	static TSharedRef<FUICommandDragDropOp> New( TSharedRef<const FUICommandInfo> InCommandInfo, FName InOriginMultiBox, TSharedPtr<SWidget> CustomDectorator, FVector2D DecoratorOffset );
+	static TSharedRef<FUICommandDragDropOp> New( FName InItemName, EMultiBlockType InBlockType, bool bInIsDraggingSection, FName InOriginMultiBox, TSharedPtr<SWidget> CustomDectorator, FVector2D DecoratorOffset );
 
-	FUICommandDragDropOp( TSharedRef<const FUICommandInfo> InUICommand, FName InOriginMultiBox, TSharedPtr<SWidget> InCustomDecorator, FVector2D DecoratorOffset )
-		: UICommand( InUICommand )
+	FUICommandDragDropOp( FName InItemName, EMultiBlockType InBlockType, bool bInIsDraggingSection, FName InOriginMultiBox, TSharedPtr<SWidget> InCustomDecorator, FVector2D DecoratorOffset )
+		: ItemName( InItemName )
+		, BlockType( InBlockType )
+		, bIsDraggingSection(bInIsDraggingSection)
 		, OriginMultiBox( InOriginMultiBox )
 		, CustomDecorator( InCustomDecorator )
 		, Offset( DecoratorOffset )
@@ -39,8 +42,14 @@ public:
 
 public:
 
-	/** UI command being dragged */
-	TSharedPtr< const FUICommandInfo > UICommand;
+	/** UI entry being dragged */
+	FName ItemName;
+
+	/** UI entry type being dragged */
+	EMultiBlockType BlockType;
+
+	/** UI entry being dragged is a section header or section separator */
+	bool bIsDraggingSection;
 
 	/** Multibox the UI command was dragged from if any*/
 	FName OriginMultiBox;
diff --git a/Engine/Source/Runtime/Slate/Public/Framework/MultiBox/MultiBox.h b/Engine/Source/Runtime/Slate/Public/Framework/MultiBox/MultiBox.h
--- a/Engine/Source/Runtime/Slate/Public/Framework/MultiBox/MultiBox.h
+++ b/Engine/Source/Runtime/Slate/Public/Framework/MultiBox/MultiBox.h
@@ -22,6 +22,7 @@ class SHorizontalBox;
 class SMultiBoxWidget;
 class STableViewBase;
 class SVerticalBox;
+class UToolMenuBase;
 
 namespace MultiBoxConstants
 {	
@@ -60,12 +61,13 @@ public:
 	 * @param InCommand		The command info that describes what action to take when this block is activated
 	 * @param InCommandList	The list of mappings from command info to delegates so we can find the delegates to process for the provided action
 	 */
-	FMultiBlock( const TSharedPtr< const FUICommandInfo > InCommand, TSharedPtr< const FUICommandList > InCommandList, FName InExtensionHook = NAME_None, EMultiBlockType InType = EMultiBlockType::None )
+	FMultiBlock( const TSharedPtr< const FUICommandInfo > InCommand, TSharedPtr< const FUICommandList > InCommandList, FName InExtensionHook = NAME_None, EMultiBlockType InType = EMultiBlockType::None, bool bInIsPartOfHeading = false)
 		: Action( InCommand )
 		, ActionList( InCommandList )
 		, ExtensionHook( InExtensionHook )
 		, Type( InType )
 		, TutorialHighlightName( NAME_None )
+		, bIsPartOfHeading(bInIsPartOfHeading)
 	{
 	}
 
@@ -74,11 +76,12 @@ public:
 	 *
 	 * @InAction UI action delegates that should be used in place of UI commands (dynamic menu items)
 	 */
-	FMultiBlock( const FUIAction& InAction,  FName InExtensionHook = NAME_None, EMultiBlockType InType = EMultiBlockType::None )
+	FMultiBlock( const FUIAction& InAction,  FName InExtensionHook = NAME_None, EMultiBlockType InType = EMultiBlockType::None, bool bInIsPartOfHeading = false )
 		: DirectActions( InAction )
 		, ExtensionHook( InExtensionHook )
 		, Type( InType )
 		, TutorialHighlightName( NAME_None )
+		, bIsPartOfHeading(bInIsPartOfHeading)
 	{
 	}
 
@@ -147,6 +150,26 @@ public:
 		return Type;
 	}
 
+	/**
+	 * Is this block a separator
+	 *
+	 * @return	True if block is a separator
+	 */
+	bool IsSeparator() const
+	{
+		return (Type == EMultiBlockType::MenuSeparator) || (Type == EMultiBlockType::ToolBarSeparator);
+	}
+
+	/**
+	 * Is this block a heading block or a block that belongs to a heading such as a separator
+	 *
+	 * @return	True if block is part of a section heading's blocks
+	 */
+	bool IsPartOfHeading() const
+	{
+		return (Type == EMultiBlockType::Heading) || bIsPartOfHeading;
+	}
+
 	/**
 	* Sets the searchable state of this block
 	*
@@ -160,6 +183,9 @@ public:
 	*/
 	bool GetSearchable() const;
 
+	/** Gets the extension hook so users can see what hooks are where */
+	FName GetExtensionHook() const { return ExtensionHook; }
+
 private:
 	/**
 	 * Allocates a widget for this type of MultiBlock.  Override this in derived classes.
@@ -167,9 +193,6 @@ private:
 	 * @return  MultiBlock widget object
 	 */
 	virtual TSharedRef< class IMultiBlockBaseWidget > ConstructWidget() const = 0;
-	
-	/** Gets the extension hook so users can see what hooks are where */
-	FName GetExtensionHook() const {return ExtensionHook;}
 
 private:
 
@@ -196,6 +219,9 @@ private:
 
 	/** Whether this block can be searched */
 	bool bSearchable;
+
+	/** Whether this block is part of the heading blocks for a section */
+	bool bIsPartOfHeading;
 };
 
 
@@ -257,6 +283,13 @@ public:
 
 	DECLARE_DELEGATE_TwoParams( FOnMakeMultiBoxBuilderOverride, const TSharedRef<FMultiBox>&, const TSharedRef<SMultiBoxWidget>& );
 
+	DECLARE_DELEGATE_RetVal_ThreeParams(TSharedRef<SWidget>, FOnModifyBlockWidgetAfterMake, const TSharedRef<SMultiBoxWidget>&, const FMultiBlock&, const TSharedRef<SWidget>& );
+
+	/**
+	 * Allow further modifications to the block's widget after it has been made
+	 */
+	FOnModifyBlockWidgetAfterMake ModifyBlockWidgetAfterMake;
+
 	/**
 	 * Creates a MultiBox widget for this MultiBox
 	 *
@@ -300,17 +333,39 @@ public:
 	TSharedPtr<FMultiBlock> MakeMultiBlockFromCommand( TSharedPtr<const FUICommandInfo> Command, bool bCommandMustBeBound ) const;
 
 	/**
-	 * Finds an existing block that handles the provided command
+	 * Finds an existing block by name and type
 	 *
-	 * @param Command The command to check for
+	 * @param InName The name to search for
+	  * @param InType The type to match during the search
+	 */
+	TSharedPtr<const FMultiBlock> FindBlockFromNameAndType(const FName InName, const EMultiBlockType InType) const;
+
+	/** @return Is being editing */
+	bool IsInEditMode() const;
+
+	/** @return The tool menu associated with this multi box */
+	UToolMenuBase* GetToolMenu() const;
+
+	/**
+	 * Only callable during edit mode
+	 * @return Index of section heading or separator
 	 */
-	TSharedPtr<const FMultiBlock> FindBlockFromCommand( TSharedPtr<const FUICommandInfo> Command ) const;
+	int32 GetSectionEditBounds(const int32 Index, int32& OutSectionEndIndex) const;
 
-	bool IsInEditMode() const { return FMultiBoxSettings::IsInToolbarEditMode() && IsCustomizable(); }
+	DECLARE_DELEGATE_OneParam(FEditSelectionChangedDelegate, TSharedRef<const FMultiBlock>);
+
+	/** Delegate to call while editing when selected block has changed */
+	const FEditSelectionChangedDelegate& OnEditSelectionChanged() const { return EditSelectionChanged; }
+
+	/** Delegate to call while editing when selected block has changed */
+	FEditSelectionChangedDelegate& OnEditSelectionChanged() { return EditSelectionChanged; }
 
 	/* The search widget to be displayed at the top of the multibox */
 	TSharedPtr<STextBlock> SearchTextWidget;
 
+	/** Weak reference to tool menu that created this multibox */
+	TWeakObjectPtr<UToolMenuBase> WeakToolMenu;
+
 private:
 	
 	/**
@@ -326,15 +381,8 @@ private:
 	 */
 	bool IsCustomizable() const;
 
-	/**
-	 * Applies user customized blocks to the multibox
-	 */
-	void ApplyCustomizedBlocks();
 private:
 
-	/** Saved customization data */
-	TSharedRef< class FMultiBoxCustomizationData > CustomizationData;
-
 	/** All command lists in this box */
 	TArray< TSharedPtr<const FUICommandList> > CommandLists;
 
@@ -350,6 +398,9 @@ private:
 	/** Type of MultiBox */
 	EMultiBoxType Type;
 
+	/** Delegate to call while editing when selected block has changed */
+	FEditSelectionChangedDelegate EditSelectionChanged;
+
 	/** True if window that owns any widgets created from this multibox should be closed automatically after the user commits to a menu choice */
 	bool bShouldCloseWindowAfterMenuSelection;
 };
@@ -417,6 +468,10 @@ public:
 	 */
 	virtual EMultiBlockLocation::Type GetMultiBlockLocation() = 0;
 
+	/**
+	 * Returns true if editing this widget
+	 */
+	virtual bool IsInEditMode() const = 0;
 };
 
 
@@ -438,6 +493,7 @@ public:
 	virtual void SetMultiBlock( TSharedRef< const FMultiBlock > InMultiBlock ) override;
 	virtual void SetMultiBlockLocation( EMultiBlockLocation::Type InLocation, bool bInSectionContainsIcons ) override;
 	virtual EMultiBlockLocation::Type GetMultiBlockLocation() override;
+	virtual bool IsInEditMode() const override;
 
 	/** SWidget Interface */
 	virtual void OnDragEnter( const FGeometry& MyGeometry, const FDragDropEvent& DragDropEvent ) override;
@@ -627,34 +683,38 @@ public:
 	*/
 	void AddSearchElement( TSharedPtr<SWidget>, FText );
 
-private:
-	/** Adds a block Widget to this widget */
-	void AddBlockWidget( const FMultiBlock& Block, TSharedPtr<SHorizontalBox> HorizontalBox, TSharedPtr<SVerticalBox> VerticalBox, EMultiBlockLocation::Type InLocation, bool bSectionContainsIcons );
-
 	/**
 	 * @return True if the passed in block is being dragged
 	 */
 	bool IsBlockBeingDragged( TSharedPtr<const FMultiBlock> Block ) const;
 
 	/**
-	 * Updates the preview block being dragged.  The drag area is where the users dragged block will be dropped
+	 * @return The visibility of customization widgets for a block
 	 */
-	void UpdateDropAreaPreviewBlock( TSharedRef<const FMultiBlock> MultiBlock, TSharedPtr<FUICommandDragDropOp> DragDropContent, const FGeometry& DragArea, const FVector2D& DragPos );
+	EVisibility GetCustomizationVisibility( TWeakPtr<const FMultiBlock> BlockWeakPtr, TWeakPtr<SWidget> BlockWidgetWeakPtr ) const;
 
 	/**
-   	 * @return The visibility of customization widgets for a block
+	 * @return The visibility of the drop location indicator of a drag and drop for a block
 	 */
-	EVisibility GetCustomizationVisibility( TWeakPtr<const FMultiBlock> BlockWeakPtr, TWeakPtr<SWidget> BlockWidgetWeakPtr ) const;
-	
-	/** Called when a user clicks the delete button on a block */
-	FReply OnDeleteBlockClicked( TWeakPtr<const FMultiBlock> BlockWeakPtr );
+	EVisibility GetCustomizationBorderDragVisibility(const FName InBlockName, const EMultiBlockType InBlockType, bool& bOutInsertAfter) const;
+
+private:
+	/** Adds a block Widget to this widget */
+	void AddBlockWidget( const FMultiBlock& Block, TSharedPtr<SHorizontalBox> HorizontalBox, TSharedPtr<SVerticalBox> VerticalBox, EMultiBlockLocation::Type InLocation, bool bSectionContainsIcons );
+
+	/**
+	 * Updates the preview block being dragged.  The drag area is where the users dragged block will be dropped
+	 */
+	void UpdateDropAreaPreviewBlock( TSharedRef<const FMultiBlock> MultiBlock, TSharedPtr<FUICommandDragDropOp> DragDropContent, const FGeometry& DragArea, const FVector2D& DragPos );
 
 private:
 	/** A preview of a block being dragged */
 	struct FDraggedMultiBlockPreview
 	{
-		/** Command being dragged */
-		TSharedPtr<const FUICommandInfo> UICommand;
+		/** Name of entry being dragged */
+		FName BlockName;
+		/** Type of entry being dragged */
+		EMultiBlockType BlockType;
 		/** Preview block for the command */
 		TSharedPtr<class FDropPreviewBlock> PreviewBlock;
 		/** Index into the block list where the block will be added*/
@@ -663,19 +723,26 @@ private:
 		EOrientation InsertOrientation;
 
 		FDraggedMultiBlockPreview()
-			: InsertIndex( INDEX_NONE )
+			: BlockType(EMultiBlockType::None)
+			, InsertIndex( INDEX_NONE )
 		{}
 
 		void Reset()
 		{
-			UICommand.Reset();
+			BlockName = NAME_None;
+			BlockType = EMultiBlockType::None;
 			PreviewBlock.Reset();
 			InsertIndex = INDEX_NONE;
 		}
 
-		bool IsValid() const { return UICommand.IsValid() && PreviewBlock.IsValid() && InsertIndex != INDEX_NONE; }
+		bool IsSameBlockAs(const FName InName, const EMultiBlockType InType) const
+		{
+			return BlockName == InName && BlockType == InType;
+		}
+
+		bool IsValid() const { return BlockName != NAME_None && BlockType != EMultiBlockType::None && PreviewBlock.IsValid() && InsertIndex != INDEX_NONE; }
 	};
-	
+
 	/** The MultiBox we're associated with */
 	TSharedPtr< FMultiBox > MultiBox;
 
diff --git a/Engine/Source/Runtime/Slate/Public/Framework/MultiBox/MultiBoxDefs.h b/Engine/Source/Runtime/Slate/Public/Framework/MultiBox/MultiBoxDefs.h
--- a/Engine/Source/Runtime/Slate/Public/Framework/MultiBox/MultiBoxDefs.h
+++ b/Engine/Source/Runtime/Slate/Public/Framework/MultiBox/MultiBoxDefs.h
@@ -69,18 +69,9 @@ public:
 
 	FMultiBoxSettings();
 
-	/** Enables or disables toolbar editing */
-	static void ToggleToolbarEditing();
-
-	/** @return True if in toolbar edit mode */
-	static bool IsInToolbarEditMode() { return bInToolbarEditMode; }
-
 	static TSharedRef< SToolTip > ConstructDefaultToolTip( const TAttribute<FText>& ToolTipText, const TSharedPtr<SWidget>& OverrideContent, const TSharedPtr<const FUICommandInfo>& Action );
 
 	static void ResetToolTipConstructor();
-
-private:
-	static bool bInToolbarEditMode;
 };
 
 struct SLATE_API FMultiBoxCustomization
diff --git a/Engine/Source/Runtime/Slate/Public/Framework/MultiBox/ToolMenuBase.h b/Engine/Source/Runtime/Slate/Public/Framework/MultiBox/ToolMenuBase.h
--- a/Engine/Source/Runtime/Slate/Public/Framework/MultiBox/ToolMenuBase.h
+++ b/Engine/Source/Runtime/Slate/Public/Framework/MultiBox/ToolMenuBase.h
@@ -8,25 +8,49 @@
 
 class FMultiBox;
 
+UENUM()
+enum class ECustomizedToolMenuVisibility
+{
+	None,
+	Visible,
+	Hidden
+};
+
 USTRUCT()
 struct FCustomizedToolMenuEntry
 {
 	GENERATED_BODY()
 
+	FCustomizedToolMenuEntry() :
+		Visibility(ECustomizedToolMenuVisibility::None)
+	{
+	}
+
 	UPROPERTY()
-	FName Name;
+	ECustomizedToolMenuVisibility Visibility;
 };
 
 USTRUCT()
 struct FCustomizedToolMenuSection
 {
 	GENERATED_BODY()
-		
+
+	FCustomizedToolMenuSection() :
+		Visibility(ECustomizedToolMenuVisibility::None)
+	{
+	}
+
 	UPROPERTY()
-	FName Name;
+	ECustomizedToolMenuVisibility Visibility;
+};
+
+USTRUCT()
+struct FCustomizedToolMenuNameArray
+{
+	GENERATED_BODY()
 
 	UPROPERTY()
-	TArray<FCustomizedToolMenuEntry> Entries;
+	TArray<FName> Names;
 };
 
 USTRUCT()
@@ -38,27 +62,41 @@ struct SLATE_API FCustomizedToolMenu
 	FName Name;
 
 	UPROPERTY()
-	TArray<FCustomizedToolMenuSection> Sections;
+	TMap<FName, FCustomizedToolMenuEntry> Entries;
 
 	UPROPERTY()
-	TArray<FName> HiddenSections;
+	TMap<FName, FCustomizedToolMenuSection> Sections;
 
 	UPROPERTY()
-	TArray<FName> HiddenEntries;
+	TMap<FName, FCustomizedToolMenuNameArray> EntryOrder;
 
-	bool IsSectionHidden(const FName InSectionName) const;
-	bool IsEntryHidden(const FName InSectionName) const;
+	UPROPERTY()
+	TArray<FName> SectionOrder;
+
+	FCustomizedToolMenuEntry* FindEntry(const FName InEntryName);
+	const FCustomizedToolMenuEntry* FindEntry(const FName InEntryName) const;
+	FCustomizedToolMenuEntry* AddEntry(const FName InEntryName);
+	ECustomizedToolMenuVisibility GetEntryVisiblity(const FName InSectionName) const;
+	bool IsEntryHidden(const FName InEntryName) const;
+	FName GetEntrySectionName(const FName InEntryName) const;
 
 	FCustomizedToolMenuSection* FindSection(const FName InSectionName);
 	const FCustomizedToolMenuSection* FindSection(const FName InSectionName) const;
+	FCustomizedToolMenuSection* AddSection(const FName InSectionName);
+	ECustomizedToolMenuVisibility GetSectionVisiblity(const FName InSectionName) const;
+	bool IsSectionHidden(const FName InSectionName) const;
+};
 
-	FCustomizedToolMenuEntry* FindEntry(const FName InEntryName, FName* OutSectionName = nullptr);
-	const FCustomizedToolMenuEntry* FindEntry(const FName InEntryName, FName* OutSectionName = nullptr) const;
+struct SLATE_API FCustomizedToolMenuHierarchy
+{
+	FName GetEntrySectionName(const FName InEntryName) const;
+	bool IsEntryHidden(const FName InEntryName) const;
+
+	bool IsSectionHidden(const FName InSectionName) const;
 
-	/* Updates re-positioning of sections and entries
-	 * Does not need to update hidden state as MultiBlocks/Widgets do not store hidden state of each entry and section themselves
-	 */
-	void UpdateFromMultiBox(const TSharedRef<const FMultiBox>& InMultiBox);
+	FCustomizedToolMenu GenerateFlattened() const;
+
+	TArray<const FCustomizedToolMenu*> Hierarchy;
 };
 
 UCLASS(Abstract)
@@ -71,6 +109,12 @@ public:
 	virtual bool IsEditing() const { return false; }
 	virtual FName GetSectionName(const FName InEntryName) const { return NAME_None; }
 
+	virtual bool ContainsSection(const FName InName) const { return false; }
+	virtual bool ContainsEntry(const FName InName) const { return false; }
+
 	virtual FCustomizedToolMenu* FindMenuCustomization() const { return nullptr; }
 	virtual FCustomizedToolMenu* AddMenuCustomization() const { return nullptr; }
+	virtual FCustomizedToolMenuHierarchy GetMenuCustomizationHierarchy() const { return FCustomizedToolMenuHierarchy(); }
+	virtual void UpdateMenuCustomizationFromMultibox(const TSharedRef<const FMultiBox>& InMultiBox) {}
 };
+
