Fix pch compile error
#rnx

diff --git a/Engine/Source/Editor/ToolMenusEditor/Private/ToolMenusEditorModule.cpp b/Engine/Source/Editor/ToolMenusEditor/Private/ToolMenusEditorModule.cpp
--- a/Engine/Source/Editor/ToolMenusEditor/Private/ToolMenusEditorModule.cpp
+++ b/Engine/Source/Editor/ToolMenusEditor/Private/ToolMenusEditorModule.cpp
@@ -7,6 +7,7 @@
 #include "IToolMenusEditorModule.h"
 #include "SEditToolMenuDialog.h"
 #include "ToolMenus.h"
+#include "Framework/Docking/TabManager.h"
 
 #define LOCTEXT_NAMESPACE "ToolMenusEditor"
 
