Menu editor
- Add minimum window size
- Update icon in menu
- Add reset all customizations button with prompt
#rnx

diff --git a/Engine/Source/Developer/ToolMenus/Private/ToolMenus.cpp b/Engine/Source/Developer/ToolMenus/Private/ToolMenus.cpp
--- a/Engine/Source/Developer/ToolMenus/Private/ToolMenus.cpp
+++ b/Engine/Source/Developer/ToolMenus/Private/ToolMenus.cpp
@@ -919,7 +919,8 @@ void UToolMenus::PopulateToolBarBuilder(FToolBarBuilder& ToolBarBuilder, UToolMe
 			}), 
 			"MenuName",
 			LOCTEXT("EditMenu", "Edit Menu"),
-			LOCTEXT("EditMenu_ToolTip", "Open menu editor")
+			LOCTEXT("EditMenu_ToolTip", "Open menu editor"),
+			EditToolbarIcon
 		);
 		ToolBarBuilder.EndSection();
 	}
@@ -1782,3 +1783,7 @@ void UToolMenus::SaveCustomizations()
 	SaveConfig();
 }
 
+void UToolMenus::RemoveAllCustomizations()
+{
+	CustomizedMenus.Reset();
+}
diff --git a/Engine/Source/Developer/ToolMenus/Public/ToolMenus.h b/Engine/Source/Developer/ToolMenus/Public/ToolMenus.h
--- a/Engine/Source/Developer/ToolMenus/Public/ToolMenus.h
+++ b/Engine/Source/Developer/ToolMenus/Public/ToolMenus.h
@@ -259,6 +259,12 @@ public:
 	/** Remove customization for a menu */
 	void RemoveCustomization(const FName InName);
 
+	/** Remove all menu customizations for all menus */
+	void RemoveAllCustomizations();
+
+	/** Save menu customizations to ini files */
+	void SaveCustomizations();
+
 	/** Find customization settings for a menu */
 	FCustomizedToolMenu* FindMenuCustomization(const FName InName);
 
@@ -288,15 +294,15 @@ public:
 	/** Icon to display in menus for command to open menu editor */
 	FSlateIcon EditMenuIcon;
 
+	/** Icon to display in toolbars for command to open menu editor */
+	FSlateIcon EditToolbarIcon;
+
 	/** Join two paths together */
 	static FName JoinMenuPaths(const FName Base, const FName Child);
 
 	/** Break apart a menu path into components */
 	static bool SplitMenuPath(const FName MenuPath, FName& OutLeft, FName& OutRight);
 
-	/** Save menu customizations to ini files */
-	void SaveCustomizations();
-
 	friend struct FToolMenuOwnerScoped;
 	friend struct FToolMenuStringCommand;
 
diff --git a/Engine/Source/Editor/ToolMenusEditor/Private/SEditToolMenuDialog.cpp b/Engine/Source/Editor/ToolMenusEditor/Private/SEditToolMenuDialog.cpp
--- a/Engine/Source/Editor/ToolMenusEditor/Private/SEditToolMenuDialog.cpp
+++ b/Engine/Source/Editor/ToolMenusEditor/Private/SEditToolMenuDialog.cpp
@@ -16,6 +16,7 @@
 #include "Widgets/Input/SButton.h"
 #include "Widgets/Input/SCheckBox.h"
 #include "Widgets/Layout/SScrollBox.h"
+#include "Dialogs/Dialogs.h"
 
 #include "Framework/Docking/TabManager.h"
 #include "EditorStyleSet.h"
@@ -41,7 +42,7 @@ public:
 
 	/** Construct this widget */
 	void Construct(const FArguments& InArgs, TSharedRef<SMultiBoxWidget> InBaseWidget, TSharedRef<const FMultiBlock> InBlock, TSharedRef<SEditToolMenuDialog> InDialogWidget)
-	{	
+	{
 		BaseWidget = InBaseWidget;
 		Block = InBlock;
 		DialogWidget = InDialogWidget;
@@ -636,6 +637,16 @@ void SEditToolMenuDialog::BuildWidget()
 					.OnClicked(this, &SEditToolMenuDialog::UndoAllChanges)
 				]
 
+				+SHorizontalBox::Slot()
+				.Padding(8, 0, 0, 0)
+				.AutoWidth()
+				[
+					SNew(SButton)
+					.Text(LOCTEXT("ResetAllDefaultsButtonText", "Reset all to Defaults"))
+					.ToolTipText(LOCTEXT("ResetAllDefaultsButtonTooltip", "Reset all settings for all menus menu to default values"))
+					.OnClicked(this, &SEditToolMenuDialog::HandleResetAllClicked)
+				]
+
 				+SHorizontalBox::Slot()
 				.Padding(8, 0, 0, 0)
 				.AutoWidth()
@@ -879,6 +890,20 @@ FReply SEditToolMenuDialog::HandleResetClicked()
 	return FReply::Handled();
 }
 
+FReply SEditToolMenuDialog::HandleResetAllClicked()
+{
+	if (OpenMsgDlgInt(EAppMsgType::YesNo,
+		LOCTEXT("ResetAllQuestion", "Remove all menu customizations for all menus?"),
+		LOCTEXT("ResetAllQuestion_Question", "Question")) == EAppReturnType::Yes)
+	{
+		UToolMenus::Get()->RemoveAllCustomizations();
+		OriginalSettings.Reset();
+		Refresh();
+	}
+
+	return FReply::Handled();
+}
+
 void SEditToolMenuDialog::SaveSettingsToDisk()
 {
 	UToolMenus::Get()->SaveCustomizations();
diff --git a/Engine/Source/Editor/ToolMenusEditor/Private/SEditToolMenuDialog.h b/Engine/Source/Editor/ToolMenusEditor/Private/SEditToolMenuDialog.h
--- a/Engine/Source/Editor/ToolMenusEditor/Private/SEditToolMenuDialog.h
+++ b/Engine/Source/Editor/ToolMenusEditor/Private/SEditToolMenuDialog.h
@@ -57,6 +57,7 @@ private:
 
 	FReply Refresh();
 	FReply HandleResetClicked();
+	FReply HandleResetAllClicked();
 	FReply UndoAllChanges();
 
 	void SaveSettingsToDisk();
diff --git a/Engine/Source/Editor/ToolMenusEditor/Private/ToolMenusEditorModule.cpp b/Engine/Source/Editor/ToolMenusEditor/Private/ToolMenusEditorModule.cpp
--- a/Engine/Source/Editor/ToolMenusEditor/Private/ToolMenusEditorModule.cpp
+++ b/Engine/Source/Editor/ToolMenusEditor/Private/ToolMenusEditorModule.cpp
@@ -74,7 +74,9 @@ public:
 			.Title(WindowTitle)
 			.ClientSize(WindowSize)
 			.SupportsMinimize(true)
-			.SupportsMaximize(true);
+			.SupportsMaximize(true)
+			.MinWidth(620)
+			.MinHeight(260);
 
 		TSharedRef<SEditToolMenuDialog> NewDialog =
 			SNew(SEditToolMenuDialog)
diff --git a/Engine/Source/Editor/UnrealEd/Private/EditorEngine.cpp b/Engine/Source/Editor/UnrealEd/Private/EditorEngine.cpp
--- a/Engine/Source/Editor/UnrealEd/Private/EditorEngine.cpp
+++ b/Engine/Source/Editor/UnrealEd/Private/EditorEngine.cpp
@@ -786,7 +786,8 @@ void UEditorEngine::InitEditor(IEngineLoop* InEngineLoop)
 
 	if (FSlateApplication::IsInitialized() && UToolMenus::IsToolMenuUIEnabled())
 	{
-		UToolMenus::Get()->EditMenuIcon = FSlateIcon(FEditorStyle::GetStyleSetName(), "ContentBrowser.AssetActions.OpenInExternalEditor");
+		UToolMenus::Get()->EditMenuIcon = FSlateIcon(FCoreStyle::Get().GetStyleSetName(), "MultiBox.GenericToolBarIcon.Small");
+		UToolMenus::Get()->EditToolbarIcon = FSlateIcon(FCoreStyle::Get().GetStyleSetName(), "MultiBox.GenericToolBarIcon");
 
 		TWeakPtr<FTimerManager> WeakTimerManager = TimerManager;
 		UToolMenus::Get()->AssignSetTimerForNextTickDelegate(FSimpleDelegate::CreateLambda([WeakTimerManager]()
