Menu editor reset now updates details panel
#rnx

diff --git a/Engine/Source/Editor/ToolMenusEditor/Private/SEditToolMenuDialog.cpp b/Engine/Source/Editor/ToolMenusEditor/Private/SEditToolMenuDialog.cpp
--- a/Engine/Source/Editor/ToolMenusEditor/Private/SEditToolMenuDialog.cpp
+++ b/Engine/Source/Editor/ToolMenusEditor/Private/SEditToolMenuDialog.cpp
@@ -369,8 +369,6 @@ void SEditToolMenuDialog::InitMenu(UToolMenu* InMenu)
 
 void SEditToolMenuDialog::Construct( const FArguments& InArgs )
 {
-	bClosedWithOk = false;
-
 	InitMenu(InArgs._SourceMenu);
 
 	SetSelectedItem(NAME_None, ESelectedEditMenuEntryType::Menu);
@@ -843,7 +841,11 @@ void SEditToolMenuDialog::OnToggleVisibleClicked(TSharedRef<const FMultiBlock> I
 		MenuCustomization->AddEntry(ItemName)->Visibility = bHidden ? ECustomizedToolMenuVisibility::Visible : ECustomizedToolMenuVisibility::Hidden;
 	}
 
-	// Update details panel
+	LoadSelectedObjectState();
+}
+
+void SEditToolMenuDialog::LoadSelectedObjectState()
+{
 	if (SelectedObject.IsValid())
 	{
 		SelectedObject->LoadState();
@@ -858,8 +860,7 @@ void SEditToolMenuDialog::OnToggleVisibleClicked(TSharedRef<const FMultiBlock> I
 void SEditToolMenuDialog::CloseContainingWindow()
 {
 	TSharedPtr<SWindow> ContainingWindow = FSlateApplication::Get().FindWidgetWindow(AsShared());
-
-	if ( ContainingWindow.IsValid() )
+	if (ContainingWindow.IsValid())
 	{
 		ContainingWindow->RequestDestroyWindow();
 	}
@@ -874,7 +875,9 @@ FReply SEditToolMenuDialog::Refresh()
 	}
 
 	BuildWidget();
-	
+
+	LoadSelectedObjectState();
+
 	return FReply::Handled();
 }
 
diff --git a/Engine/Source/Editor/ToolMenusEditor/Private/SEditToolMenuDialog.h b/Engine/Source/Editor/ToolMenusEditor/Private/SEditToolMenuDialog.h
--- a/Engine/Source/Editor/ToolMenusEditor/Private/SEditToolMenuDialog.h
+++ b/Engine/Source/Editor/ToolMenusEditor/Private/SEditToolMenuDialog.h
@@ -60,6 +60,7 @@ private:
 	FReply HandleResetAllClicked();
 	FReply UndoAllChanges();
 
+	void LoadSelectedObjectState();
 	void SaveSettingsToDisk();
 
 	TSharedRef<SWidget> BuildMenuPropertiesWidget();
@@ -70,8 +71,6 @@ private:
 
 	void SetSelectedItem(const FName InName, ESelectedEditMenuEntryType InType);
 
-	bool bClosedWithOk;
-
 	TWeakObjectPtr<UToolMenu> MenuDialogOpenedWith;
 	TWeakObjectPtr<UToolMenu> CurrentGeneratedMenu;
 
