Changed the background of SEditToolMenuDialog.
Docking.Tab.Normal doesn't to exist anymore in editor style.
And MultiBox.Visible ... and MultiBox.NotVisible ... doesn't exist.

diff --git a/Engine/Source/Editor/ToolMenusEditor/Private/SEditToolMenuDialog.cpp b/Engine/Source/Editor/ToolMenusEditor/Private/SEditToolMenuDialog.cpp
--- a/Engine/Source/Editor/ToolMenusEditor/Private/SEditToolMenuDialog.cpp
+++ b/Engine/Source/Editor/ToolMenusEditor/Private/SEditToolMenuDialog.cpp
@@ -89,15 +89,15 @@ private:
 	{
 		if (IsVisible())
 		{
-			static const FName NAME_VisibleHoveredBrush("MultiBox.VisibleHighlightIcon16x");
-			static const FName NAME_VisibleNotHoveredBrush("MultiBox.VisibleIcon16x");
+			static const FName NAME_VisibleHoveredBrush("Level.VisibleHighlightIcon16x");
+			static const FName NAME_VisibleNotHoveredBrush("Level.VisibleIcon16x");
 			return IsHovered() ? FEditorStyle::GetBrush(NAME_VisibleHoveredBrush) :
 				FEditorStyle::GetBrush(NAME_VisibleNotHoveredBrush);
 		}
 		else
 		{
-			static const FName NAME_NotVisibleHoveredBrush("MultiBox.NotVisibleHighlightIcon16x");
-			static const FName NAME_NotVisibleNotHoveredBrush("MultiBox.NotVisibleIcon16x");
+			static const FName NAME_NotVisibleHoveredBrush("Level.NotVisibleHighlightIcon16x");
+			static const FName NAME_NotVisibleNotHoveredBrush("Level.NotVisibleIcon16x");
 			return IsHovered() ? FEditorStyle::GetBrush(NAME_NotVisibleHoveredBrush) :
 				FEditorStyle::GetBrush(NAME_NotVisibleNotHoveredBrush);
 		}
@@ -674,7 +674,7 @@ void SEditToolMenuDialog::BuildWidget()
 	[
 		SNew(SBorder)
 		.Padding(20)
-		.BorderImage( FEditorStyle::GetBrush("Docking.Tab.Normal") )
+		.BorderImage( FEditorStyle::GetBrush("Docking.Tab.ContentAreaBrush") )
 		[
 			SNew(SVerticalBox)
 
