Editor - ToolMenus - Show additional information in ToolMenus.Edit dialog
#rnx

diff --git a/Engine/Source/Editor/ToolMenusEditor/Private/ToolMenusEditor.cpp b/Engine/Source/Editor/ToolMenusEditor/Private/ToolMenusEditor.cpp
--- a/Engine/Source/Editor/ToolMenusEditor/Private/ToolMenusEditor.cpp
+++ b/Engine/Source/Editor/ToolMenusEditor/Private/ToolMenusEditor.cpp
@@ -59,6 +59,39 @@ void UToolMenuEditorDialogEntry::Init(UToolMenu* InMenu, const FName InName)
 	Type = ESelectedEditMenuEntryType::Entry;
 	Menu = InMenu;
 
+	OwnerName = TEXT("");
+	ScriptObject = TEXT("");
+	ScriptObjectClass = TEXT("");
+
+	if (Menu)
+	{
+		for (const FToolMenuSection& Section : Menu->Sections)
+		{
+			if (const FToolMenuEntry* FoundEntry = Section.FindEntry(Name))
+			{
+				if (FoundEntry->Owner.IsSet())
+				{
+					if (!FoundEntry->Owner.TryGetName().IsNone())
+					{
+						OwnerName = FoundEntry->Owner.TryGetName().ToString();
+					}
+					else
+					{
+						OwnerName = TEXT("<Pointer>");
+					}
+				}
+
+				if (UToolMenuEntryScript* EntryScriptObject = FoundEntry->ScriptObject)
+				{
+					ScriptObject = EntryScriptObject->GetFullName();
+					ScriptObjectClass = EntryScriptObject->GetClass()->GetFullName();
+				}
+
+				break;
+			}
+		}
+	}
+
 	LoadState();
 }
 
diff --git a/Engine/Source/Editor/ToolMenusEditor/Public/ToolMenusEditor.h b/Engine/Source/Editor/ToolMenusEditor/Public/ToolMenusEditor.h
--- a/Engine/Source/Editor/ToolMenusEditor/Public/ToolMenusEditor.h
+++ b/Engine/Source/Editor/ToolMenusEditor/Public/ToolMenusEditor.h
@@ -60,6 +60,15 @@ public:
 	UPROPERTY(EditAnywhere, Category = Misc)
 	ECustomizedToolMenuVisibility Visibility;
 
+	UPROPERTY(VisibleAnywhere, Category = Misc)
+	FString OwnerName;
+
+	UPROPERTY(VisibleAnywhere, Category = Misc)
+	FString ScriptObject;
+
+	UPROPERTY(VisibleAnywhere, Category = Misc)
+	FString ScriptObjectClass;
+
 	UPROPERTY()
 	TObjectPtr<UToolMenu> Menu;
 };
